Added subsystem to support advanced stylus input for the editor such as Tilt - enable the Stylus & Tablet Input plugin to use it.

Input can be listened to by registered a message handler in UStylusInputSubsystem.

Currently only Windows is supported, using the IRealTimeStylus COM interface.

diff --git a/Engine/Plugins/Editor/StylusInput/Source/StylusInput/Private/SStylusInputDebugWidget.cpp b/Engine/Plugins/Editor/StylusInput/Source/StylusInput/Private/SStylusInputDebugWidget.cpp
new file mode 100644
--- /dev/null
+++ b/Engine/Plugins/Editor/StylusInput/Source/StylusInput/Private/SStylusInputDebugWidget.cpp
@@ -0,0 +1,220 @@
+// Copyright 1998-2019 Epic Games, Inc. All Rights Reserved.
+
+#include "SStylusInputDebugWidget.h"
+
+#include "IStylusState.h"
+
+#include "Widgets/SBoxPanel.h"
+#include "Widgets/Text/STextBlock.h"
+
+#define LOCTEXT_NAMESPACE "StylusInputDebugWidget"
+
+SStylusInputDebugWidget::SStylusInputDebugWidget()
+{
+}
+
+SStylusInputDebugWidget::~SStylusInputDebugWidget()
+{
+	InputSubsystem->RemoveMessageHandler(*this);
+}
+
+void SStylusInputDebugWidget::Construct(const FArguments& InArgs, UStylusInputSubsystem& InSubsystem)
+{
+	InputSubsystem = &InSubsystem;
+	InputSubsystem->AddMessageHandler(*this);
+
+	ChildSlot
+	[
+		SNew(SVerticalBox)
+		+ SVerticalBox::Slot()
+		.AutoHeight()
+		[
+			SNew(SHorizontalBox)
+			+ SHorizontalBox::Slot()
+				.FillWidth(1)
+			[
+				SNew(STextBlock)
+				.Text(LOCTEXT("MostRecentIndex", "Most Recent Index"))
+			]
+			+ SHorizontalBox::Slot()
+				.FillWidth(1)
+			[
+				SNew(STextBlock)
+				.Text(this, &SStylusInputDebugWidget::GetIndexText)
+			]
+		]
+		+SVerticalBox::Slot()
+			.AutoHeight()
+		[
+			SNew(SHorizontalBox)
+			+SHorizontalBox::Slot()
+				.FillWidth(1)
+			[
+				SNew(STextBlock)
+				.Text(LOCTEXT("Position", "Position"))
+			]
+			+SHorizontalBox::Slot()
+				.FillWidth(1)
+			[
+				SNew(STextBlock)
+				.Text(this, &SStylusInputDebugWidget::GetPositionText)
+			]
+		]
+		+SVerticalBox::Slot()
+			.AutoHeight()
+		[
+			SNew(SHorizontalBox)
+			+SHorizontalBox::Slot()
+				.FillWidth(1)
+			[
+				SNew(STextBlock)
+				.Text(LOCTEXT("NormalPressure", "Normal Pressure"))
+			]
+		+ SHorizontalBox::Slot()
+			.FillWidth(1)
+			[
+				SNew(STextBlock)
+				.Text(this, &SStylusInputDebugWidget::GetPressureText)
+			]
+		]
+		+SVerticalBox::Slot()
+			.AutoHeight()
+			[
+			SNew(SHorizontalBox)
+			+ SHorizontalBox::Slot()
+				.FillWidth(1)
+			[
+				SNew(STextBlock)
+				.Text(LOCTEXT("TangentPressure", "Tangent Pressure"))
+			]
+			+SHorizontalBox::Slot()
+				.FillWidth(1)
+			[
+				SNew(STextBlock)
+				.Text(this, &SStylusInputDebugWidget::GetTangentPressureText)
+			]
+		]
+		+SVerticalBox::Slot()
+			.AutoHeight()
+		[
+			SNew(SHorizontalBox)
+			+SHorizontalBox::Slot()
+				.FillWidth(1)
+			[
+				SNew(STextBlock)
+				.Text(LOCTEXT("Z", "Z"))
+			]
+			+SHorizontalBox::Slot()
+				.FillWidth(1)
+			[
+				SNew(STextBlock)
+				.Text(this, &SStylusInputDebugWidget::GetZText)
+			]
+		]
+		+SVerticalBox::Slot()
+			.AutoHeight()
+		[
+			SNew(SHorizontalBox)
+			+SHorizontalBox::Slot()
+				.FillWidth(1)
+			[
+				SNew(STextBlock)
+				.Text(LOCTEXT("Twist", "Twist"))
+			]
+			+SHorizontalBox::Slot()
+				.FillWidth(1)
+			[
+				SNew(STextBlock)
+				.Text(this, &SStylusInputDebugWidget::GetTwistText)
+			]
+		]
+		+SVerticalBox::Slot()
+			.AutoHeight()
+		[
+			SNew(SHorizontalBox)
+			+SHorizontalBox::Slot()
+				.FillWidth(1)
+			[
+				SNew(STextBlock)
+				.Text(LOCTEXT("Tilt", "Tilt"))
+			]
+			+SHorizontalBox::Slot()
+				.FillWidth(1)
+			[
+				SNew(STextBlock)
+				.Text(this, &SStylusInputDebugWidget::GetTiltText)
+			]
+		]
+		+SVerticalBox::Slot()
+			.AutoHeight()
+		[
+			SNew(SHorizontalBox)
+			+ SHorizontalBox::Slot()
+				.FillWidth(1)
+			[
+				SNew(STextBlock)
+				.Text(LOCTEXT("Size", "Size"))
+			]
+			+SHorizontalBox::Slot()
+			.FillWidth(1)
+			[
+				SNew(STextBlock)
+				.Text(this, &SStylusInputDebugWidget::GetSizeText)
+			]
+		]
+		+SVerticalBox::Slot()
+			.AutoHeight()
+		[
+			SNew(SHorizontalBox)
+			+SHorizontalBox::Slot()
+				.FillWidth(1)
+			[
+				SNew(STextBlock)
+				.Text(LOCTEXT("IsTouching", "Is Touching?"))
+			]
+			+SHorizontalBox::Slot()
+				.FillWidth(1)
+			[
+				SNew(SCheckBox)
+				.IsChecked(this, &SStylusInputDebugWidget::IsTouching)
+			]
+		]
+		+SVerticalBox::Slot()
+			.AutoHeight()
+		[
+			SNew(SHorizontalBox)
+			+SHorizontalBox::Slot()
+				.FillWidth(1)
+			[
+				SNew(STextBlock)
+				.Text(LOCTEXT("IsInverted", "Is Inverted?"))
+			]
+			+SHorizontalBox::Slot()
+				.FillWidth(1)
+			[
+				SNew(SCheckBox)
+				.IsChecked(this, &SStylusInputDebugWidget::IsInverted)	
+			]
+		]
+	];
+}
+
+FText SStylusInputDebugWidget::GetVector2Text(FVector2D Value)
+{
+	return FText::FromString(FString::Format(TEXT("{0}, {1}"), { Value.X, Value.Y }));
+}
+
+FText SStylusInputDebugWidget::GetFloatText(float Value)
+{
+	return FText::FromString(FString::Format(TEXT("{0}"), { Value }));
+}
+
+ECheckBoxState SStylusInputDebugWidget::IsTouching() const
+{
+	return State.IsStylusDown() ? ECheckBoxState::Checked : ECheckBoxState::Unchecked;
+}
+
+ECheckBoxState SStylusInputDebugWidget::IsInverted() const
+{
+	return State.IsStylusInverted() ? ECheckBoxState::Checked : ECheckBoxState::Unchecked;
+}
\ No newline at end of file
diff --git a/Engine/Plugins/Editor/StylusInput/Source/StylusInput/Private/SStylusInputDebugWidget.h b/Engine/Plugins/Editor/StylusInput/Source/StylusInput/Private/SStylusInputDebugWidget.h
new file mode 100644
--- /dev/null
+++ b/Engine/Plugins/Editor/StylusInput/Source/StylusInput/Private/SStylusInputDebugWidget.h
@@ -0,0 +1,50 @@
+// Copyright 1998-2019 Epic Games, Inc. All Rights Reserved.
+
+#pragma once
+
+#include "Widgets/SCompoundWidget.h"
+
+#include "IStylusInputModule.h"
+#include "IStylusState.h"
+
+class SStylusInputDebugWidget : public SCompoundWidget, 
+	public IStylusMessageHandler
+{
+public:
+	SStylusInputDebugWidget();
+	virtual ~SStylusInputDebugWidget();
+
+	SLATE_BEGIN_ARGS(SStylusInputDebugWidget)
+	{}
+	SLATE_END_ARGS();
+
+	void Construct(const FArguments& InArgs, UStylusInputSubsystem& InSubsystem);
+	void OnStylusStateChanged(const FStylusState& InState, int32 InIndex)
+	{
+		State = InState;
+		LastIndex = InIndex;
+	}
+
+private:
+
+	UStylusInputSubsystem* InputSubsystem;
+	FStylusState State;
+	int32 LastIndex;
+
+	ECheckBoxState IsTouching() const;
+	ECheckBoxState IsInverted() const;
+
+	FText GetPositionText() const { return GetVector2Text(State.GetPosition()); }
+	FText GetTiltText() const { return GetVector2Text(State.GetTilt()); }
+	FText GetSizeText() const { return GetVector2Text(State.GetSize()); }
+
+	FText GetIndexText() const { return FText::FromString(FString::FromInt(LastIndex)); }
+
+	FText GetPressureText() const { return GetFloatText(State.GetPressure()); }
+	FText GetTangentPressureText() const { return GetFloatText(State.GetTangentPressure()); }
+	FText GetZText() const { return GetFloatText(State.GetZ()); }
+	FText GetTwistText() const { return GetFloatText(State.GetTwist()); }
+
+	static FText GetVector2Text(FVector2D Value);
+	static FText GetFloatText(float Value);
+};
\ No newline at end of file
diff --git a/Engine/Plugins/Editor/StylusInput/Source/StylusInput/Private/StylusInputModule.cpp b/Engine/Plugins/Editor/StylusInput/Source/StylusInput/Private/StylusInputModule.cpp
new file mode 100644
--- /dev/null
+++ b/Engine/Plugins/Editor/StylusInput/Source/StylusInput/Private/StylusInputModule.cpp
@@ -0,0 +1,125 @@
+// Copyright 1998-2019 Epic Games, Inc. All Rights Reserved.
+
+#include "IStylusInputModule.h"
+#include "CoreMinimal.h"
+
+#include "Framework/Docking/TabManager.h"
+#include "Interfaces/IMainFrameModule.h"
+#include "Logging/LogMacros.h"
+#include "Widgets/Docking/SDockTab.h"
+#include "Widgets/SWindow.h"
+#include "WorkspaceMenuStructure.h"
+#include "WorkspaceMenuStructureModule.h"
+
+DEFINE_LOG_CATEGORY_STATIC(LogStylusInput, Log, All);
+
+static const FName StylusInputDebugTabName = FName("StylusInputDebug");
+
+class FStylusInputModule : public IModuleInterface
+{
+};
+
+IMPLEMENT_MODULE(FStylusInputModule, StylusInput)
+
+// This is the interface that all platform-specific implementations must implement.
+class IStylusInputInterfaceInternal
+{
+public:
+	virtual IStylusInputDevice* GetInputDevice(int32 Index) const = 0;
+	virtual int32 NumInputDevices() const = 0;
+};
+
+// This is the function that all platform-specific implementations are required to implement.
+TSharedPtr<IStylusInputInterfaceInternal> CreateStylusInputInterface();
+
+#if PLATFORM_WINDOWS
+#include "WindowsStylusInputInterface.h"
+#endif // PLATFORM_WINDOWS
+
+// TODO: Other platforms
+
+void UStylusInputSubsystem::Initialize(FSubsystemCollectionBase& Collection)
+{
+	Super::Initialize(Collection);
+
+	UE_LOG(LogStylusInput, Log, TEXT("Initializing StylusInput subsystem."));
+
+	InputInterface = CreateStylusInputInterface();
+
+	const TSharedRef<FGlobalTabmanager>& TabManager = FGlobalTabmanager::Get();
+	const IWorkspaceMenuStructure& MenuStructure = WorkspaceMenu::GetMenuStructure();
+
+	TabManager->RegisterNomadTabSpawner(StylusInputDebugTabName,
+		FOnSpawnTab::CreateUObject(this, &UStylusInputSubsystem::OnSpawnPluginTab))
+		.SetDisplayName(LOCTEXT("DebugTabTitle", "Stylus Input Debug"))
+		.SetTooltipText(LOCTEXT("DebugTabTooltip", "Debug panel to display current values of stylus inputs."))
+		.SetGroup(MenuStructure.GetDeveloperToolsMiscCategory());
+}
+
+void UStylusInputSubsystem::Deinitialize()
+{
+	Super::Deinitialize();
+
+	FGlobalTabmanager::Get()->UnregisterTabSpawner(StylusInputDebugTabName);
+
+	InputInterface.Reset();
+
+	UE_LOG(LogStylusInput, Log, TEXT("Shutting down StylusInput subsystem."));
+}
+
+int32 UStylusInputSubsystem::NumInputDevices() const
+{
+	if (InputInterface.IsValid())
+	{
+		return InputInterface->NumInputDevices();
+	}
+	return 0;
+}
+
+const IStylusInputDevice* UStylusInputSubsystem::GetInputDevice(int32 Index) const
+{
+	if (InputInterface.IsValid())
+	{
+		return InputInterface->GetInputDevice(Index);
+	}
+	return nullptr;
+}
+
+void UStylusInputSubsystem::AddMessageHandler(IStylusMessageHandler& InHandler)
+{
+	MessageHandlers.AddUnique(&InHandler);
+}
+
+void UStylusInputSubsystem::RemoveMessageHandler(IStylusMessageHandler& InHandler)
+{
+	MessageHandlers.Remove(&InHandler);
+}
+
+void UStylusInputSubsystem::Tick(float DeltaTime)
+{
+	if (InputInterface.IsValid())
+	{
+		for (int32 DeviceIdx = 0; DeviceIdx < NumInputDevices(); ++DeviceIdx)
+		{
+			IStylusInputDevice* InputDevice = InputInterface->GetInputDevice(DeviceIdx);
+			if (InputDevice->IsDirty())
+			{
+				InputDevice->Tick();
+
+				for (IStylusMessageHandler* Handler : MessageHandlers)
+				{
+					Handler->OnStylusStateChanged(InputDevice->GetCurrentState(), DeviceIdx);
+				}
+			}
+		}
+	}
+}
+
+TSharedRef<SDockTab> UStylusInputSubsystem::OnSpawnPluginTab(const FSpawnTabArgs& Args)
+{
+	return SNew(SDockTab)
+		.TabRole(ETabRole::NomadTab)
+		[
+			SNew(SStylusInputDebugWidget, *this)
+		];
+}
diff --git a/Engine/Plugins/Editor/StylusInput/Source/StylusInput/Private/WindowsRealTimeStylusPlugin.cpp b/Engine/Plugins/Editor/StylusInput/Source/StylusInput/Private/WindowsRealTimeStylusPlugin.cpp
new file mode 100644
--- /dev/null
+++ b/Engine/Plugins/Editor/StylusInput/Source/StylusInput/Private/WindowsRealTimeStylusPlugin.cpp
@@ -0,0 +1,390 @@
+// Copyright 1998-2019 Epic Games, Inc. All Rights Reserved.
+
+#include "WindowsRealTimeStylusPlugin.h"
+
+HRESULT FWindowsRealTimeStylusPlugin::QueryInterface(const IID& InterfaceID, void** Pointer)
+{
+	if ((InterfaceID == __uuidof(IStylusSyncPlugin)) || (InterfaceID == IID_IUnknown))
+	{
+		*Pointer = this;
+		AddRef();
+		return S_OK;
+	}
+	else if ((InterfaceID == IID_IMarshal) && (FreeThreadedMarshaller != nullptr))
+	{
+		return FreeThreadedMarshaller->QueryInterface(InterfaceID, Pointer);
+	}
+
+	*Pointer = nullptr;
+	return E_NOINTERFACE;
+}
+
+HRESULT FWindowsRealTimeStylusPlugin::StylusDown(IRealTimeStylus* RealTimeStylus, const StylusInfo* StylusInfo, ULONG PacketSize, LONG* Packet, LONG** InOutPackets)
+{
+	FTabletContextInfo* TabletContext = FindTabletContext(StylusInfo->tcid);
+	if (TabletContext != nullptr)
+	{
+		TabletContext->WindowsState.IsTouching = true;
+	}
+	return S_OK;
+}
+
+HRESULT FWindowsRealTimeStylusPlugin::StylusUp(IRealTimeStylus* RealTimeStylus, const StylusInfo* StylusInfo, ULONG PacketSize, LONG* Packet, LONG** InOutPackets)
+{
+	FTabletContextInfo* TabletContext = FindTabletContext(StylusInfo->tcid);
+	if (TabletContext != nullptr)
+	{
+		// we know this is not touching
+		TabletContext->WindowsState.IsTouching = false;
+		TabletContext->WindowsState.NormalPressure = 0;
+	}
+	return S_OK;
+}
+
+static void SetupPacketDescriptions(IRealTimeStylus* RealTimeStylus, FTabletContextInfo& TabletContext)
+{
+	ULONG NumPacketProperties;
+	PACKET_PROPERTY* PacketProperties;
+	HRESULT hr = RealTimeStylus->GetPacketDescriptionData(TabletContext.ID, nullptr, nullptr, &NumPacketProperties, &PacketProperties);
+	if (SUCCEEDED(hr))
+	{
+		for (ULONG PropIdx = 0; PropIdx < NumPacketProperties; ++PropIdx)
+		{
+			PACKET_PROPERTY CurrentProperty = PacketProperties[PropIdx];
+
+			EWindowsPacketType PacketType = EWindowsPacketType::None;
+			if (CurrentProperty.guid == GUID_PACKETPROPERTY_GUID_X)
+			{
+				PacketType = EWindowsPacketType::X;
+			}
+			else if (CurrentProperty.guid == GUID_PACKETPROPERTY_GUID_Y)
+			{
+				PacketType = EWindowsPacketType::Y;
+			}
+			else if (CurrentProperty.guid == GUID_PACKETPROPERTY_GUID_Z)
+			{
+				PacketType = EWindowsPacketType::Z;
+			}
+			else if (CurrentProperty.guid == GUID_PACKETPROPERTY_GUID_PACKET_STATUS)
+			{
+				PacketType = EWindowsPacketType::Status;
+			}
+			else if (CurrentProperty.guid == GUID_PACKETPROPERTY_GUID_NORMAL_PRESSURE)
+			{
+				PacketType = EWindowsPacketType::NormalPressure;
+			}
+			else if (CurrentProperty.guid == GUID_PACKETPROPERTY_GUID_TANGENT_PRESSURE)
+			{
+				PacketType = EWindowsPacketType::TangentPressure;
+			}
+			else if (CurrentProperty.guid == GUID_PACKETPROPERTY_GUID_BUTTON_PRESSURE)
+			{
+				PacketType = EWindowsPacketType::ButtonPressure;
+			}
+			else if (CurrentProperty.guid == GUID_PACKETPROPERTY_GUID_ALTITUDE_ORIENTATION)
+			{
+				PacketType = EWindowsPacketType::Altitude;
+			}
+			else if (CurrentProperty.guid == GUID_PACKETPROPERTY_GUID_AZIMUTH_ORIENTATION)
+			{
+				PacketType = EWindowsPacketType::Azimuth;
+			}
+			else if (CurrentProperty.guid == GUID_PACKETPROPERTY_GUID_TWIST_ORIENTATION)
+			{
+				PacketType = EWindowsPacketType::Twist;
+			}
+			else if (CurrentProperty.guid == GUID_PACKETPROPERTY_GUID_X_TILT_ORIENTATION)
+			{
+				PacketType = EWindowsPacketType::XTilt;
+			}
+			else if (CurrentProperty.guid == GUID_PACKETPROPERTY_GUID_Y_TILT_ORIENTATION)
+			{
+				PacketType = EWindowsPacketType::YTilt;
+			}
+			else if (CurrentProperty.guid == GUID_PACKETPROPERTY_GUID_WIDTH)
+			{
+				PacketType = EWindowsPacketType::Width;
+			}
+			else if (CurrentProperty.guid == GUID_PACKETPROPERTY_GUID_HEIGHT)
+			{
+				PacketType = EWindowsPacketType::Height;
+			}
+
+			int32 CreatedIdx = TabletContext.PacketDescriptions.Emplace();
+			FPacketDescription& PacketDescription = TabletContext.PacketDescriptions[CreatedIdx];
+			PacketDescription.Type = PacketType;
+			PacketDescription.Minimum = CurrentProperty.PropertyMetrics.nLogicalMin;
+			PacketDescription.Maximum = CurrentProperty.PropertyMetrics.nLogicalMax;
+			PacketDescription.Resolution = CurrentProperty.PropertyMetrics.fResolution;
+		}
+
+		::CoTaskMemFree(PacketProperties);
+	}
+}
+
+static void SetupTabletSupportedPackets(TComPtr<IRealTimeStylus> RealTimeStylus, FTabletContextInfo& TabletContext)
+{
+	IInkTablet* InkTablet;
+	RealTimeStylus->GetTabletFromTabletContextId(TabletContext.ID, &InkTablet);
+
+	int16 Supported;
+	InkTablet->IsPacketPropertySupported(STR_GUID_X, &Supported);
+	if (Supported)
+	{
+		TabletContext.SupportedPackets.Add(EWindowsPacketType::X);
+	}
+
+	InkTablet->IsPacketPropertySupported(STR_GUID_Y, &Supported);
+	if (Supported)
+	{
+		TabletContext.SupportedPackets.Add(EWindowsPacketType::Y);
+	}
+
+	InkTablet->IsPacketPropertySupported(STR_GUID_Z, &Supported);
+	if (Supported)
+	{
+		TabletContext.SupportedPackets.Add(EWindowsPacketType::Z);
+		TabletContext.AddSupportedInput(EStylusInputType::Z);
+	}
+
+	InkTablet->IsPacketPropertySupported(STR_GUID_PAKETSTATUS, &Supported);
+	if (Supported)
+	{
+		TabletContext.SupportedPackets.Add(EWindowsPacketType::Status);
+	}
+
+	InkTablet->IsPacketPropertySupported(STR_GUID_NORMALPRESSURE, &Supported);
+	if (Supported)
+	{
+		TabletContext.SupportedPackets.Add(EWindowsPacketType::NormalPressure);
+		TabletContext.AddSupportedInput(EStylusInputType::Pressure);
+	}
+
+	InkTablet->IsPacketPropertySupported(STR_GUID_TANGENTPRESSURE, &Supported);
+	if (Supported)
+	{
+		TabletContext.SupportedPackets.Add(EWindowsPacketType::TangentPressure);
+		TabletContext.AddSupportedInput(EStylusInputType::TangentPressure);
+	}
+
+	InkTablet->IsPacketPropertySupported(STR_GUID_BUTTONPRESSURE, &Supported);
+	if (Supported)
+	{
+		TabletContext.SupportedPackets.Add(EWindowsPacketType::ButtonPressure);
+		TabletContext.AddSupportedInput(EStylusInputType::ButtonPressure);
+	}
+
+	InkTablet->IsPacketPropertySupported(STR_GUID_AZIMUTHORIENTATION, &Supported);
+	if (Supported)
+	{
+		TabletContext.SupportedPackets.Add(EWindowsPacketType::Azimuth);
+	}
+
+	InkTablet->IsPacketPropertySupported(STR_GUID_ALTITUDEORIENTATION, &Supported);
+	if (Supported)
+	{
+		TabletContext.SupportedPackets.Add(EWindowsPacketType::Altitude);
+	}
+
+	InkTablet->IsPacketPropertySupported(STR_GUID_XTILTORIENTATION, &Supported);
+	if (Supported)
+	{
+		TabletContext.SupportedPackets.Add(EWindowsPacketType::XTilt);
+	}
+
+	InkTablet->IsPacketPropertySupported(STR_GUID_YTILTORIENTATION, &Supported);
+	if (Supported)
+	{
+		TabletContext.SupportedPackets.Add(EWindowsPacketType::YTilt);
+	}
+
+	InkTablet->IsPacketPropertySupported(STR_GUID_TWISTORIENTATION, &Supported);
+	if (Supported)
+	{
+		TabletContext.SupportedPackets.Add(EWindowsPacketType::Twist);
+		TabletContext.AddSupportedInput(EStylusInputType::Twist);
+	}
+
+	InkTablet->IsPacketPropertySupported(STR_GUID_WIDTH, &Supported);
+	if (Supported)
+	{
+		TabletContext.SupportedPackets.Add(EWindowsPacketType::Width);
+	}
+
+	InkTablet->IsPacketPropertySupported(STR_GUID_HEIGHT, &Supported);
+	if (Supported)
+	{
+		TabletContext.SupportedPackets.Add(EWindowsPacketType::Height);
+	}
+
+	if (TabletContext.SupportedPackets.Contains(EWindowsPacketType::X) &&
+		TabletContext.SupportedPackets.Contains(EWindowsPacketType::Y))
+	{
+		TabletContext.AddSupportedInput(EStylusInputType::Position);
+	}
+
+	if (TabletContext.SupportedPackets.Contains(EWindowsPacketType::XTilt) &&
+		TabletContext.SupportedPackets.Contains(EWindowsPacketType::YTilt))
+	{
+		TabletContext.AddSupportedInput(EStylusInputType::Tilt);
+	}
+
+	if (TabletContext.SupportedPackets.Contains(EWindowsPacketType::Width) &&
+		TabletContext.SupportedPackets.Contains(EWindowsPacketType::Height))
+	{
+		TabletContext.AddSupportedInput(EStylusInputType::Size);
+	}
+}
+
+FTabletContextInfo* FWindowsRealTimeStylusPlugin::FindTabletContext(TABLET_CONTEXT_ID TabletID)
+{
+	for (FTabletContextInfo& TabletContext : TabletContexts)
+	{
+		if (TabletContext.ID == TabletID)
+		{
+			return &TabletContext;
+		}
+	}
+	return nullptr;
+}
+
+void FWindowsRealTimeStylusPlugin::AddTabletContext(IRealTimeStylus* RealTimeStylus, TABLET_CONTEXT_ID TabletID)
+{
+	FTabletContextInfo* FoundContext = FindTabletContext(TabletID);
+	if (FoundContext == nullptr)
+	{
+		int32 Created = TabletContexts.Emplace();
+		FoundContext = &TabletContexts[Created];
+		FoundContext->ID = TabletID;
+	}
+
+	SetupTabletSupportedPackets(RealTimeStylus, *FoundContext);
+	SetupPacketDescriptions(RealTimeStylus, *FoundContext);
+}
+
+void FWindowsRealTimeStylusPlugin::RemoveTabletContext(IRealTimeStylus* RealTimeStylus, TABLET_CONTEXT_ID TabletID)
+{
+	for (int32 ExistingIdx = 0; ExistingIdx < TabletContexts.Num(); ++ExistingIdx)
+	{
+		if (TabletContexts[ExistingIdx].ID == TabletID)
+		{
+			TabletContexts.RemoveAt(ExistingIdx);
+			break;
+		}
+	}
+}
+
+HRESULT FWindowsRealTimeStylusPlugin::RealTimeStylusEnabled(IRealTimeStylus* RealTimeStylus, ULONG Num, const TABLET_CONTEXT_ID* InTabletContexts)
+{
+	for (ULONG TabletIdx = 0; TabletIdx < Num; ++TabletIdx)
+	{
+		AddTabletContext(RealTimeStylus, InTabletContexts[TabletIdx]);
+	}
+	return S_OK;
+}
+
+HRESULT FWindowsRealTimeStylusPlugin::RealTimeStylusDisabled(IRealTimeStylus* RealTimeStylus, ULONG Num, const TABLET_CONTEXT_ID* InTabletContexts)
+{
+	for (ULONG TabletIdx = 0; TabletIdx < Num; ++TabletIdx)
+	{
+		RemoveTabletContext(RealTimeStylus, InTabletContexts[TabletIdx]);
+	}
+	return S_OK;
+}
+
+HRESULT FWindowsRealTimeStylusPlugin::TabletAdded(IRealTimeStylus* RealTimeStylus, IInkTablet* InkTablet)
+{
+	TABLET_CONTEXT_ID TabletID;
+	if (SUCCEEDED(RealTimeStylus->GetTabletContextIdFromTablet(InkTablet, &TabletID)))
+	{
+		AddTabletContext(RealTimeStylus, TabletID);
+	}
+	return S_OK;
+}
+
+HRESULT FWindowsRealTimeStylusPlugin::TabletRemoved(IRealTimeStylus* RealTimeStylus, LONG iTabletIndex)
+{
+	TabletContexts.RemoveAt(iTabletIndex);
+	return S_OK;
+}
+
+static float Normalize(int Value, const FPacketDescription& Desc)
+{
+	return (float) (Value - Desc.Minimum) / (float) ((Desc.Maximum - Desc.Minimum) + 1);
+}
+
+static float ToDegrees(int Value, const FPacketDescription& Desc)
+{
+	return Value / Desc.Resolution;
+}
+
+void FWindowsRealTimeStylusPlugin::HandlePacket(IRealTimeStylus* RealTimeStylus, const StylusInfo* StylusInfo, ULONG PacketCount, ULONG PacketBufferLength, LONG* Packets)
+{
+	FTabletContextInfo* TabletContext = FindTabletContext(StylusInfo->tcid);
+	if (TabletContext == nullptr)
+	{
+		return;
+	}
+
+	TabletContext->SetDirty();
+	TabletContext->WindowsState.IsInverted = StylusInfo->bIsInvertedCursor;
+
+	ULONG PropertyCount = PacketBufferLength / PacketCount;
+
+	for (ULONG i = 0; i < PropertyCount; ++i)
+	{
+		const FPacketDescription& PacketDescription = TabletContext->PacketDescriptions[i];
+
+		float Normalized = Normalize(Packets[i], PacketDescription);
+
+		switch (PacketDescription.Type)
+		{
+			case EWindowsPacketType::X:
+				TabletContext->WindowsState.Position.X = Packets[i];
+				break;
+			case EWindowsPacketType::Y:
+				TabletContext->WindowsState.Position.Y = Packets[i];
+				break;
+			case EWindowsPacketType::Status:
+				break;
+			case EWindowsPacketType::Z:
+				TabletContext->WindowsState.Z = Normalized;
+				break;
+			case EWindowsPacketType::NormalPressure:
+				TabletContext->WindowsState.NormalPressure = Normalized;
+				break;
+			case EWindowsPacketType::TangentPressure:
+				TabletContext->WindowsState.TangentPressure = Normalized;
+				break;
+			case EWindowsPacketType::Twist:
+				TabletContext->WindowsState.Twist = ToDegrees(Packets[i], PacketDescription);
+				break;
+			case EWindowsPacketType::XTilt:
+				TabletContext->WindowsState.Tilt.X = ToDegrees(Packets[i], PacketDescription);
+				break;
+			case EWindowsPacketType::YTilt:
+				TabletContext->WindowsState.Tilt.Y = ToDegrees(Packets[i], PacketDescription);
+				break;
+			case EWindowsPacketType::Width:
+				TabletContext->WindowsState.Size.X = Normalized;
+				break;
+			case EWindowsPacketType::Height:
+				TabletContext->WindowsState.Size.Y = Normalized;
+				break;
+		}
+	}
+}
+
+HRESULT FWindowsRealTimeStylusPlugin::Packets(IRealTimeStylus* RealTimeStylus, const StylusInfo* StylusInfo,
+	ULONG PacketCount, ULONG PacketBufferLength, LONG* Packets, ULONG* InOutPackets, LONG** PtrInOutPackets)
+{
+	HandlePacket(RealTimeStylus, StylusInfo, PacketCount, PacketBufferLength, Packets);
+	return S_OK;
+}
+
+HRESULT FWindowsRealTimeStylusPlugin::InAirPackets(IRealTimeStylus* RealTimeStylus, const StylusInfo* StylusInfo,
+	ULONG PacketCount, ULONG PacketBufferLength, LONG* Packets, ULONG* InOutPackets, LONG** PtrInOutPackets)
+{
+	HandlePacket(RealTimeStylus, StylusInfo, PacketCount, PacketBufferLength, Packets);
+	return S_OK;
+}
\ No newline at end of file
diff --git a/Engine/Plugins/Editor/StylusInput/Source/StylusInput/Private/WindowsRealTimeStylusPlugin.h b/Engine/Plugins/Editor/StylusInput/Source/StylusInput/Private/WindowsRealTimeStylusPlugin.h
new file mode 100644
--- /dev/null
+++ b/Engine/Plugins/Editor/StylusInput/Source/StylusInput/Private/WindowsRealTimeStylusPlugin.h
@@ -0,0 +1,171 @@
+// Copyright 1998-2019 Epic Games, Inc. All Rights Reserved.
+
+#pragma once
+
+#include "Windows/WindowsHWrapper.h"
+
+#include "Windows/AllowWindowsPlatformTypes.h"
+#include "Windows/PreWindowsApi.h"
+#include "Windows/COMPointer.h"
+	#include <guiddef.h>
+	#include <RTSCom.h>
+#include "Windows/PostWindowsApi.h"
+#include "Windows/HideWindowsPlatformTypes.h"
+
+#include "IStylusState.h"
+
+/**
+ * Packet types as derived from IRealTimeStylus::GetPacketDescriptionData.
+ */
+enum class EWindowsPacketType
+{
+	None,
+	X,
+	Y,
+	Z,
+	Status,
+	NormalPressure,
+	TangentPressure,
+	ButtonPressure,
+	Azimuth,
+	Altitude,
+	Twist,
+	XTilt,
+	YTilt,
+	Width,
+	Height,
+};
+
+/**
+ * Stylus state for a single frame.
+ */
+struct FWindowsStylusState
+{
+	FVector2D Position;
+	float Z;
+	FVector2D Tilt;
+	float Twist;
+	float NormalPressure;
+	float TangentPressure;
+	FVector2D Size;
+	bool IsTouching : 1;
+	bool IsInverted : 1;
+
+	FWindowsStylusState() :
+		Position(0, 0), Z(0), Tilt(0, 0), Twist(0), NormalPressure(0), TangentPressure(0),
+		Size(0, 0), IsTouching(false), IsInverted(false)
+	{
+	}
+
+	FStylusState ToPublicState() const
+	{
+		return FStylusState(Position, Z, Tilt, Twist, NormalPressure, TangentPressure, Size, IsTouching, IsInverted);
+	}
+};
+
+/**
+ * Description of a packet's information, as derived from IRealTimeStylus::GetPacketDescriptionData.
+ */
+struct FPacketDescription
+{
+	EWindowsPacketType Type { EWindowsPacketType::None };
+	int32 Minimum { 0 };
+	int32 Maximum { 0 };
+	float Resolution { 0 };
+};
+
+struct FTabletContextInfo : public IStylusInputDevice
+{
+	int32 Index;
+
+	TABLET_CONTEXT_ID ID;
+	TArray<FPacketDescription> PacketDescriptions;
+	TArray<EWindowsPacketType> SupportedPackets;
+
+	FWindowsStylusState WindowsState;
+
+	void AddSupportedInput(EStylusInputType Type) { SupportedInputs.Add(Type); }
+	void SetDirty() { Dirty = true; }
+
+	virtual void Tick() override
+	{
+		PreviousState = CurrentState;
+		CurrentState = WindowsState.ToPublicState();
+		Dirty = false;
+	}
+};
+
+/**
+ * An implementation of an IStylusSyncPlugin for use with the RealTimeStylus API.
+ */
+class FWindowsRealTimeStylusPlugin : public IStylusSyncPlugin
+{
+public:
+	FWindowsRealTimeStylusPlugin() = default;
+	virtual ~FWindowsRealTimeStylusPlugin()
+	{
+		if (FreeThreadedMarshaller != nullptr)
+		{
+			FreeThreadedMarshaller->Release();
+		}
+	}
+
+	virtual ULONG AddRef() override { return ++RefCount; }
+	virtual ULONG Release() override
+	{
+		int NewRefCount = --RefCount;
+		if (NewRefCount == 0)
+			delete this;
+
+		return NewRefCount;
+	}
+
+	virtual HRESULT QueryInterface(const IID& InterfaceID, void** Pointer) override;
+
+	virtual HRESULT TabletAdded(IRealTimeStylus* RealTimeStylus, IInkTablet* InkTablet) override;
+	virtual HRESULT TabletRemoved(IRealTimeStylus* RealTimeStylus, LONG iTabletIndex) override;
+
+	virtual HRESULT RealTimeStylusEnabled(IRealTimeStylus* RealTimeStylus, ULONG Num, const TABLET_CONTEXT_ID* InTabletContexts) override;
+	virtual HRESULT RealTimeStylusDisabled(IRealTimeStylus* RealTimeStylus, ULONG Num, const TABLET_CONTEXT_ID* InTabletContexts) override;
+
+	virtual HRESULT StylusInRange(IRealTimeStylus* RealTimeStylus, TABLET_CONTEXT_ID TabletContext, STYLUS_ID StylusID) override { return S_OK; }
+	virtual HRESULT StylusOutOfRange(IRealTimeStylus* RealTimeStylus, TABLET_CONTEXT_ID TabletContext, STYLUS_ID StylusID) override { return S_OK; }
+
+	virtual HRESULT StylusDown(IRealTimeStylus* RealTimeStylus, const StylusInfo* StylusInfo, ULONG PacketSize, LONG* Packet, LONG** InOutPackets) override;
+	virtual HRESULT StylusUp(IRealTimeStylus* RealTimeStylus, const StylusInfo* StylusInfo, ULONG PacketSize, LONG* Packet, LONG** InOutPackets) override;
+
+	virtual HRESULT StylusButtonDown(IRealTimeStylus* RealTimeStylus, STYLUS_ID StylusID, const GUID* GUID, POINT* Position) override { return S_OK; }
+	virtual HRESULT StylusButtonUp(IRealTimeStylus* RealTimeStylus, STYLUS_ID StylusID, const GUID* GUID, POINT* Position) override { return S_OK; }
+
+	virtual HRESULT InAirPackets(IRealTimeStylus* RealTimeStylus, const StylusInfo* StylusInfo,
+		ULONG PacketCount, ULONG PacketBufferLength, LONG* Packets, ULONG* NumOutPackets, LONG** PtrOutPackets) override;
+	virtual HRESULT Packets(IRealTimeStylus* RealTimeStylus, const StylusInfo* StylusInfo,
+		ULONG PacketCount, ULONG PacketBufferSize, LONG* Packets, ULONG* NumOutPackets, LONG** PtrOutPackets) override;
+
+	virtual HRESULT CustomStylusDataAdded(IRealTimeStylus* RealTimeStylus, const GUID* GUID, ULONG Data, const BYTE* ByteData) override { return S_OK; }
+
+	virtual HRESULT SystemEvent(IRealTimeStylus* RealTimeStylus, TABLET_CONTEXT_ID TabletContext, STYLUS_ID StylusID, SYSTEM_EVENT EventType, SYSTEM_EVENT_DATA EventData) override { return S_OK; }
+	virtual HRESULT Error(IRealTimeStylus* RealTimeStylus, IStylusPlugin* Plugin, RealTimeStylusDataInterest DataInterest, HRESULT ErrorCode, LONG_PTR* Key) override { return S_OK; }
+
+	virtual HRESULT DataInterest(RealTimeStylusDataInterest* OutDataInterest) override
+	{
+		*OutDataInterest = RTSDI_AllData;
+		return S_OK;
+	}
+
+	virtual HRESULT UpdateMapping(IRealTimeStylus* RealTimeStylus) override { return S_OK; }
+
+	FTabletContextInfo* FindTabletContext(TABLET_CONTEXT_ID TabletID);
+
+	IUnknown* FreeThreadedMarshaller;
+	TArray<FTabletContextInfo> TabletContexts;
+	bool HasChanges { false };
+
+private:
+	int RefCount { 1 };
+
+	void HandlePacket(IRealTimeStylus* RealTimeStylus, const StylusInfo* StylusInfo, ULONG PacketCount, ULONG PacketBufferLength, LONG* Packets);
+
+	void AddTabletContext(IRealTimeStylus* RealTimeStylus, TABLET_CONTEXT_ID TabletID);
+	void RemoveTabletContext(IRealTimeStylus* RealTimeStylus, TABLET_CONTEXT_ID TabletID);
+};
\ No newline at end of file
diff --git a/Engine/Plugins/Editor/StylusInput/Source/StylusInput/Private/WindowsStylusInputInterface.cpp b/Engine/Plugins/Editor/StylusInput/Source/StylusInput/Private/WindowsStylusInputInterface.cpp
new file mode 100644
--- /dev/null
+++ b/Engine/Plugins/Editor/StylusInput/Source/StylusInput/Private/WindowsStylusInputInterface.cpp
@@ -0,0 +1,158 @@
+// Copyright 1998-2019 Epic Games, Inc. All Rights Reserved.
+
+#include "WindowsStylusInputInterface.h"
+#include "WindowsRealTimeStylusPlugin.h"
+
+class FWindowsStylusInputInterfaceImpl
+{
+public:
+	~FWindowsStylusInputInterfaceImpl();
+
+	TComPtr<IRealTimeStylus> RealTimeStylus;
+	TSharedPtr<FWindowsRealTimeStylusPlugin> StylusPlugin;
+	void* DLLHandle { nullptr };
+};
+
+FWindowsStylusInputInterface::FWindowsStylusInputInterface(FWindowsStylusInputInterfaceImpl* InImpl)
+{
+	Impl = InImpl;
+}
+
+FWindowsStylusInputInterface::~FWindowsStylusInputInterface()
+{
+	delete Impl;
+}
+
+int32 FWindowsStylusInputInterface::NumInputDevices() const
+{
+	return Impl->StylusPlugin->TabletContexts.Num();
+}
+
+IStylusInputDevice* FWindowsStylusInputInterface::GetInputDevice(int32 Index) const
+{
+	if (Index < 0 || Index >= Impl->StylusPlugin->TabletContexts.Num())
+	{
+		return nullptr;
+	}
+
+	return &Impl->StylusPlugin->TabletContexts[Index];
+}
+
+FWindowsStylusInputInterfaceImpl::~FWindowsStylusInputInterfaceImpl()
+{
+	RealTimeStylus->RemoveAllStylusSyncPlugins();
+	RealTimeStylus.Reset();
+
+	StylusPlugin.Reset();
+
+	if (DLLHandle != nullptr)
+	{
+		FPlatformProcess::FreeDllHandle(DLLHandle);
+		DLLHandle = nullptr;
+	}
+}
+
+static void OnMainFrameCreated(FWindowsStylusInputInterfaceImpl& WindowsImpl, TSharedPtr<SWindow> Window)
+{
+	TSharedPtr<FGenericWindow> NativeWindow = Window->GetNativeWindow();
+	HWND Hwnd = reinterpret_cast<HWND>(NativeWindow->GetOSWindowHandle());
+
+	WindowsImpl.RealTimeStylus->put_HWND(reinterpret_cast<uint64>(Hwnd));
+
+	// We desire to receive everything, but what we actually will receive is determined in AddTabletContext
+	TArray<GUID> DesiredPackets = {
+		GUID_PACKETPROPERTY_GUID_X,
+		GUID_PACKETPROPERTY_GUID_Y,
+		GUID_PACKETPROPERTY_GUID_Z,
+		GUID_PACKETPROPERTY_GUID_PACKET_STATUS,
+		GUID_PACKETPROPERTY_GUID_NORMAL_PRESSURE,
+		GUID_PACKETPROPERTY_GUID_TANGENT_PRESSURE,
+		GUID_PACKETPROPERTY_GUID_X_TILT_ORIENTATION,
+		GUID_PACKETPROPERTY_GUID_Y_TILT_ORIENTATION,
+		GUID_PACKETPROPERTY_GUID_TWIST_ORIENTATION,
+		GUID_PACKETPROPERTY_GUID_WIDTH,
+		GUID_PACKETPROPERTY_GUID_HEIGHT,
+		// Currently not needed.
+		//GUID_PACKETPROPERTY_GUID_BUTTON_PRESSURE,
+		//GUID_PACKETPROPERTY_GUID_AZIMUTH_ORIENTATION,
+		//GUID_PACKETPROPERTY_GUID_ALTITUDE_ORIENTATION,
+	};
+
+	WindowsImpl.RealTimeStylus->SetDesiredPacketDescription(DesiredPackets.Num(), DesiredPackets.GetData());
+
+	WindowsImpl.RealTimeStylus->put_Enabled(Windows::TRUE);
+}
+
+TSharedPtr<IStylusInputInterfaceInternal> CreateStylusInputInterface()
+{
+	FWindowsStylusInputInterfaceImpl* WindowsImpl = new FWindowsStylusInputInterfaceImpl();
+
+	if (!FWindowsPlatformMisc::CoInitialize()) 
+	{
+		UE_LOG(LogStylusInput, Error, TEXT("Could not initialize COM library!"));
+		return TSharedPtr<IStylusInputInterfaceInternal>();
+	}
+
+	// Load RealTimeStylus DLL
+	const FString InkDLLDirectory = TEXT("C:\\Program Files\\Common Files\\microsoft shared\\ink");
+	const FString RTSComDLL = TEXT("RTSCom.dll");
+	FPlatformProcess::PushDllDirectory(*InkDLLDirectory);
+
+	WindowsImpl->DLLHandle = FPlatformProcess::GetDllHandle(*(InkDLLDirectory / RTSComDLL));
+	if (WindowsImpl->DLLHandle == nullptr)
+	{
+		FWindowsPlatformMisc::CoUninitialize();
+		UE_LOG(LogStylusInput, Error, TEXT("Could not load RTSCom.dll!"));
+		return TSharedPtr<IStylusInputInterfaceInternal>();
+	}
+
+	FPlatformProcess::PopDllDirectory(*InkDLLDirectory);
+
+	// Create RealTimeStylus interface
+	void* OutInstance { nullptr };
+	HRESULT hr = ::CoCreateInstance(__uuidof(RealTimeStylus), nullptr, CLSCTX_INPROC, __uuidof(IRealTimeStylus), &OutInstance);
+	if (FAILED(hr))
+	{
+		FWindowsPlatformMisc::CoUninitialize();
+		UE_LOG(LogStylusInput, Error, TEXT("Could not create RealTimeStylus!"));
+		return TSharedPtr<IStylusInputInterfaceInternal>();
+	}
+
+	WindowsImpl->RealTimeStylus = static_cast<IRealTimeStylus*>(OutInstance);
+	WindowsImpl->StylusPlugin = MakeShareable(new FWindowsRealTimeStylusPlugin());
+	
+	// Create free-threaded marshaller for the plugin
+	hr = ::CoCreateFreeThreadedMarshaler(WindowsImpl->StylusPlugin.Get(), 
+		&WindowsImpl->StylusPlugin->FreeThreadedMarshaller);
+	if (FAILED(hr))
+	{
+		FWindowsPlatformMisc::CoUninitialize();
+		UE_LOG(LogStylusInput, Error, TEXT("Could not create FreeThreadedMarshaller!"));
+		return TSharedPtr<IStylusInputInterfaceInternal>();
+	}
+
+	// Add stylus plugin to the interface
+	hr = WindowsImpl->RealTimeStylus->AddStylusSyncPlugin(0, WindowsImpl->StylusPlugin.Get());
+	if (FAILED(hr))
+	{
+		FWindowsPlatformMisc::CoUninitialize();
+		UE_LOG(LogStylusInput, Error, TEXT("Could not add stylus plugin to API!"));
+		return TSharedPtr<IStylusInputInterfaceInternal>();
+	}
+
+	// Set hook to catch main window creation
+	IMainFrameModule& MainFrameModule = FModuleManager::LoadModuleChecked<IMainFrameModule>("MainFrame");
+	if (!MainFrameModule.IsWindowInitialized())
+	{
+		MainFrameModule.OnMainFrameCreationFinished().AddLambda([WindowsImpl](TSharedPtr<SWindow> Window, bool)
+		{
+			OnMainFrameCreated(*WindowsImpl, Window);
+		});
+	}
+	else
+	{
+		OnMainFrameCreated(*WindowsImpl, MainFrameModule.GetParentWindow());
+	}
+
+	return MakeShareable(new FWindowsStylusInputInterface(WindowsImpl));
+}
diff --git a/Engine/Plugins/Editor/StylusInput/Source/StylusInput/Private/WindowsStylusInputInterface.h b/Engine/Plugins/Editor/StylusInput/Source/StylusInput/Private/WindowsStylusInputInterface.h
new file mode 100644
--- /dev/null
+++ b/Engine/Plugins/Editor/StylusInput/Source/StylusInput/Private/WindowsStylusInputInterface.h
@@ -0,0 +1,18 @@
+#pragma once
+
+class FWindowsStylusInputInterfaceImpl;
+
+class FWindowsStylusInputInterface : public IStylusInputInterfaceInternal
+{
+public:
+	FWindowsStylusInputInterface(FWindowsStylusInputInterfaceImpl* InImpl);
+	virtual ~FWindowsStylusInputInterface();
+
+	virtual int32 NumInputDevices() const override;
+	virtual IStylusInputDevice* GetInputDevice(int32 Index) const override;
+
+private:
+	// pImpl to avoid including Windows headers.
+	FWindowsStylusInputInterfaceImpl* Impl;
+	TArray<IStylusMessageHandler*> MessageHandlers;
+};
diff --git a/Engine/Plugins/Editor/StylusInput/Source/StylusInput/Public/IStylusInputModule.h b/Engine/Plugins/Editor/StylusInput/Source/StylusInput/Public/IStylusInputModule.h
new file mode 100644
--- /dev/null
+++ b/Engine/Plugins/Editor/StylusInput/Source/StylusInput/Public/IStylusInputModule.h
@@ -0,0 +1,73 @@
+// Copyright 1998-2019 Epic Games, Inc. All Rights Reserved.
+
+#pragma once
+
+#include "IStylusState.h"
+
+#include "Modules/ModuleInterface.h"
+#include "Modules/ModuleManager.h"
+#include "EditorSubsystem.h"
+
+#include "IStylusInputModule.generated.h"
+
+/**
+ * Module to handle Wacom-style tablet input using styluses.
+ */
+class STYLUSINPUT_API IStylusInputModule : public IModuleInterface
+{
+public:
+
+	/**
+	 * Retrieve the module instance.
+	 */
+	static inline IStylusInputModule& Get()
+	{
+		return FModuleManager::LoadModuleChecked<IStylusInputModule>("StylusInput");
+	}
+
+	/**
+	 * Checks to see if this module is loaded and ready.  It is only valid to call Get() if IsAvailable() returns true.
+	 *
+	 * @return True if the module is loaded and ready to use
+	 */
+	static inline bool IsAvailable()
+	{
+		return FModuleManager::Get().IsModuleLoaded("StylusInput");
+	}
+};
+
+class IStylusInputInterfaceInternal;
+
+UCLASS()
+class STYLUSINPUT_API UStylusInputSubsystem : 
+	public UEditorSubsystem, 
+	public FTickableEditorObject
+{
+	GENERATED_BODY()
+public:
+	// UEditorSubsystem implementation
+	virtual void Initialize(FSubsystemCollectionBase& Collection) override;
+	virtual void Deinitialize() override;
+
+	/** Retrieve the input device that is at the given index, or nullptr if not found. Corresponds to the StylusIndex in IStylusMessageHandler. */
+	const IStylusInputDevice* GetInputDevice(int32 Index) const;
+
+	/** Return the number of active input devices. */
+	int32 NumInputDevices() const; 
+
+	/** Add a message handler to receive messages from the stylus. */
+	void AddMessageHandler(IStylusMessageHandler& MessageHandler);
+
+	/** Remove a previously registered message handler. */
+	void RemoveMessageHandler(IStylusMessageHandler& MessageHandler);
+
+	// FTickableEditorObject implementation
+	virtual void Tick(float DeltaTime) override;
+	virtual TStatId GetStatId() const override { RETURN_QUICK_DECLARE_CYCLE_STAT(UStylusInputSubsystem, STATGROUP_Tickables); }
+
+private:
+	TSharedPtr<IStylusInputInterfaceInternal> InputInterface;
+	TArray<IStylusMessageHandler*> MessageHandlers;
+
+	TSharedRef<SDockTab> OnSpawnPluginTab(const FSpawnTabArgs& Args);
+};
\ No newline at end of file
diff --git a/Engine/Plugins/Editor/StylusInput/Source/StylusInput/Public/IStylusState.h b/Engine/Plugins/Editor/StylusInput/Source/StylusInput/Public/IStylusState.h
new file mode 100644
--- /dev/null
+++ b/Engine/Plugins/Editor/StylusInput/Source/StylusInput/Public/IStylusState.h
@@ -0,0 +1,165 @@
+// Copyright 1998-2019 Epic Games, Inc. All Rights Reserved.
+
+#pragma once
+
+/**
+ * The types of stylus inputs that can be potentially supported by a stylus.
+ */
+enum class STYLUSINPUT_API EStylusInputType
+{
+	Position,
+	Z,
+	Pressure,
+	Tilt,
+	TangentPressure,
+	ButtonPressure,
+	Twist,
+	Size
+};
+
+/**
+ * The current state of a single stylus, as sent by IStylusMessageHandler. 
+ */
+class STYLUSINPUT_API FStylusState
+{
+public:
+
+	FStylusState() : Position(0, 0), Z(0), Tilt(0, 0), Twist(0), Pressure(0), 
+		TangentPressure(0), Size(0, 0), IsDown(false), IsInverted(false)
+	{
+	}
+
+	FStylusState(FVector2D InPosition, float InZ, FVector2D InTilt, float InTwist,
+		float InPressure, float InTanPressure, FVector2D InSize, 
+		bool InDown, bool InInverted) : 
+		Position(InPosition), Z(InZ), Tilt(InTilt), Twist(InTwist),
+		Pressure(InPressure), TangentPressure(InTanPressure), Size(InSize),
+		IsDown(InDown), IsInverted(InInverted)
+	{
+	}
+
+	FStylusState(const FStylusState& Other) :
+		Position(Other.Position), Z(Other.Z), Tilt(Other.Tilt), Twist(Other.Twist),
+		Pressure(Other.Pressure), TangentPressure(Other.TangentPressure), Size(Other.Size),
+		IsDown(Other.IsDown), IsInverted(Other.IsInverted)
+	{
+	}
+
+	/**
+	 * The current position of the stylus on (or above) the tablet. Always valid.
+	 * This value is in logical coordinates, not pixels.
+	 * The mouse position can be used to retrieve the screenspace value.
+	 * A value of (0,0) is in the top-left of the tablet.
+	 */ 
+	FVector2D GetPosition() const { return Position; }
+
+	/**
+	 * The current height of the stylus above the tablet.
+	 * Defaults to 0 if EStylusInputType::Z is not supported.
+	 */
+	float GetZ() const { return Z; }
+
+	/** 
+	 * The current tilt along the X axis in degrees, normalized to the range of [-90, 90]. 
+	 * Defaults to (0,0) if EStylusInputType::Tilt is not supported.
+	 * A value of (0,0) means that the stylus is perfectly vertical.
+	 * A positive X value means that the stylus is tilted to the right. 
+	 * A positive Y value means that the stylus is tilted forwards, away from the user.
+	 * A value of -90 or 90 means that the pen is lying on the tablet, though in practice this isn't widely supported.
+	 */
+	FVector2D GetTilt() const { return Tilt; }
+
+	/**
+	 * The current twist amount around the stylus' own axis in degrees, normalized to the range of [0, 360).
+	 * Defaults to 0 if EStylusInputType::Twist is not supported.
+	 * A value of 1 represents a full rotation clockwise.
+	 */
+	float GetTwist() const { return Twist; }
+
+	/**
+	 * Get the current pressure along the tablet's normal, usually straight down.
+	 * Normalized to the range [0, 1]. 
+	 * Defaults to 0 if EStylusInputType::Pressure is not supported.
+	 */
+	float GetPressure() const { return Pressure; }
+
+	/**
+	 * Get the current pressure along the tablet's surface.
+	 * Normalized to the range [0, 1].
+	 * Defaults to 0 if EStylusInputType::TangentPressure is not supported.
+	 */
+	float GetTangentPressure() const { return TangentPressure; }
+
+	/**
+	 * Get the size of the touch in logical coordinates.
+	 * Defaults to (0,0) if EStylusInputType::Size is not supported.
+	 */
+	FVector2D GetSize() const { return Size; }
+	
+	/**
+	 * Is the stylus inverted? Ie. the eraser part is pointing down. 
+	 * Defaults to false if EStylusInputType::Tilt is not supported.
+	 */
+	bool IsStylusInverted() const { return IsInverted; }
+
+	/** Is the stylus currently touching the tablet? */
+	bool IsStylusDown() const { return IsDown; }
+
+private:
+
+	FVector2D Position;
+	float Z;
+	FVector2D Tilt;
+	float Twist;
+	float Pressure;
+	float TangentPressure;
+	FVector2D Size;
+
+	bool IsDown : 1;
+	bool IsInverted : 1;
+};
+
+/** An input device representing a stylus and its current state. */
+class STYLUSINPUT_API IStylusInputDevice
+{
+public:
+
+	virtual ~IStylusInputDevice() {}
+
+	/**
+	 * Get the current stylus state.
+	 */
+	const FStylusState& GetCurrentState() const { return CurrentState; }
+
+	/**
+	 * Get the previous stylus state.
+	 */
+	const FStylusState& GetPreviousState() const { return PreviousState; }
+
+	/**
+	 * Get the supported inputs of this tablet.
+	 */
+	const TArray<EStylusInputType>& GetSupportedInputs() const { return SupportedInputs; }
+
+	/** Update the input device. Not intended to be called externally. */ 
+	virtual void Tick() = 0;
+
+	/** Does the input device need to be ticked? */
+	bool IsDirty() const { return Dirty; }
+
+protected:
+	FStylusState CurrentState;
+	FStylusState PreviousState;
+	TArray<EStylusInputType> SupportedInputs;
+	bool Dirty : 1;
+};
+
+/**
+ * Interface to implement for classes that want to receive messages when a stylus state change occurs.
+ * Will trigger once per frame.
+ */
+class IStylusMessageHandler
+{
+public:
+	virtual void OnStylusStateChanged(const FStylusState& NewState, int32 StylusIndex) = 0;
+};
diff --git a/Engine/Plugins/Editor/StylusInput/Source/StylusInput/StylusInput.Build.cs b/Engine/Plugins/Editor/StylusInput/Source/StylusInput/StylusInput.Build.cs
new file mode 100644
--- /dev/null
+++ b/Engine/Plugins/Editor/StylusInput/Source/StylusInput/StylusInput.Build.cs
@@ -0,0 +1,52 @@
+// Copyright 1998-2019 Epic Games, Inc. All Rights Reserved.
+
+namespace UnrealBuildTool.Rules
+{
+	public class StylusInput : ModuleRules
+	{
+		public StylusInput(ReadOnlyTargetRules Target) : base(Target)
+		{
+			PublicIncludePaths.AddRange(
+				new string[] {
+					// ... add public include paths required here ...
+				}
+				);
+
+			PrivateIncludePaths.AddRange(
+				new string[] {
+					// ... add other private include paths required here ...
+				}
+				);
+
+			PublicDependencyModuleNames.AddRange(
+				new string[]
+				{
+					"CoreUObject",
+					"EditorSubsystem",
+					"Engine",
+					"UnrealEd"
+					// ... add other public dependencies that you statically link with here ...
+				}
+				);
+
+			PrivateDependencyModuleNames.AddRange(
+				new string[]
+				{
+					"Core",
+					"MainFrame",
+					"SlateCore",
+					"Slate",
+					"WorkspaceMenuStructure"
+					// ... add private dependencies that you statically link with here ...
+				}
+				);
+
+			DynamicallyLoadedModuleNames.AddRange(
+				new string[]
+				{
+					// ... add any modules that your module loads dynamically here ...
+				}
+				);
+		}
+	}
+}
diff --git a/Engine/Plugins/Editor/StylusInput/StylusInput.uplugin b/Engine/Plugins/Editor/StylusInput/StylusInput.uplugin
new file mode 100644
--- /dev/null
+++ b/Engine/Plugins/Editor/StylusInput/StylusInput.uplugin
@@ -0,0 +1,25 @@
+{
+	"FileVersion" : 3,
+	"Version" : 1,
+	"VersionName" : "1.0",
+	"FriendlyName" : "Stylus & Tablet Plugin",
+	"Description" : "Support for advanced stylus and tablet inputs such as pressure, stylus & tablet buttons, and pen angles.",
+	"Category" : "Input Devices",
+	"CreatedBy" : "Epic Games, Inc.",
+	"CreatedByURL" : "http://epicgames.com",
+	"DocsURL" : "",
+	"MarketplaceURL" : "",
+	"SupportURL" : "",
+	"EnabledByDefault" : false,
+	"CanContainContent" : false,
+	"IsBetaVersion" : true,
+	"Installed" : false,
+	"Modules" :
+	[
+		{
+			"Name" : "StylusInput",
+			"Type" : "Editor",
+			"LoadingPhase" : "Default"
+		}
+	]
+}
\ No newline at end of file
