Fix compile error, missing include
#rnx

diff --git a/Engine/Plugins/Editor/StylusInput/Source/StylusInput/Public/IStylusState.h b/Engine/Plugins/Editor/StylusInput/Source/StylusInput/Public/IStylusState.h
--- a/Engine/Plugins/Editor/StylusInput/Source/StylusInput/Public/IStylusState.h
+++ b/Engine/Plugins/Editor/StylusInput/Source/StylusInput/Public/IStylusState.h
@@ -2,6 +2,8 @@
 
 #pragma once
 
+#include "Math/Vector2D.h"
+
 /**
  * The types of stylus inputs that can be potentially supported by a stylus.
  */
