fix compile errors in stylus module
#rnx

diff --git a/Engine/Plugins/Editor/StylusInput/Source/StylusInput/Private/SStylusInputDebugWidget.cpp b/Engine/Plugins/Editor/StylusInput/Source/StylusInput/Private/SStylusInputDebugWidget.cpp
--- a/Engine/Plugins/Editor/StylusInput/Source/StylusInput/Private/SStylusInputDebugWidget.cpp
+++ b/Engine/Plugins/Editor/StylusInput/Source/StylusInput/Private/SStylusInputDebugWidget.cpp
@@ -6,6 +6,7 @@
 
 #include "Widgets/SBoxPanel.h"
 #include "Widgets/Text/STextBlock.h"
+#include "Widgets/Input/SCheckBox.h"
 
 #define LOCTEXT_NAMESPACE "StylusInputDebugWidget"
 
diff --git a/Engine/Plugins/Editor/StylusInput/Source/StylusInput/Private/StylusInputModule.cpp b/Engine/Plugins/Editor/StylusInput/Source/StylusInput/Private/StylusInputModule.cpp
--- a/Engine/Plugins/Editor/StylusInput/Source/StylusInput/Private/StylusInputModule.cpp
+++ b/Engine/Plugins/Editor/StylusInput/Source/StylusInput/Private/StylusInputModule.cpp
@@ -10,6 +10,9 @@
 #include "Widgets/SWindow.h"
 #include "WorkspaceMenuStructure.h"
 #include "WorkspaceMenuStructureModule.h"
+#include "SStylusInputDebugWidget.h"
+
+#define LOCTEXT_NAMESPACE "FStylusInputModule"
 
 DEFINE_LOG_CATEGORY_STATIC(LogStylusInput, Log, All);
 
@@ -21,13 +24,6 @@ class FStylusInputModule : public IModuleInterface
 
 IMPLEMENT_MODULE(FStylusInputModule, StylusInput)
 
-// This is the interface that all platform-specific implementations must implement.
-class IStylusInputInterfaceInternal
-{
-public:
-	virtual IStylusInputDevice* GetInputDevice(int32 Index) const = 0;
-	virtual int32 NumInputDevices() const = 0;
-};
 
 // This is the function that all platform-specific implementations are required to implement.
 TSharedPtr<IStylusInputInterfaceInternal> CreateStylusInputInterface();
@@ -123,3 +119,7 @@ TSharedRef<SDockTab> UStylusInputSubsystem::OnSpawnPluginTab(const FSpawnTabArgs
 			SNew(SStylusInputDebugWidget, *this)
 		];
 }
+
+
+
+#undef LOCTEXT_NAMESPACE
\ No newline at end of file
diff --git a/Engine/Plugins/Editor/StylusInput/Source/StylusInput/Private/WindowsStylusInputInterface.cpp b/Engine/Plugins/Editor/StylusInput/Source/StylusInput/Private/WindowsStylusInputInterface.cpp
--- a/Engine/Plugins/Editor/StylusInput/Source/StylusInput/Private/WindowsStylusInputInterface.cpp
+++ b/Engine/Plugins/Editor/StylusInput/Source/StylusInput/Private/WindowsStylusInputInterface.cpp
@@ -2,6 +2,9 @@
 
 #include "WindowsStylusInputInterface.h"
 #include "WindowsRealTimeStylusPlugin.h"
+#include "Interfaces/IMainFrameModule.h"
+
+DEFINE_LOG_CATEGORY_STATIC(LogStylusInput, Log, All);
 
 class FWindowsStylusInputInterfaceImpl
 {
diff --git a/Engine/Plugins/Editor/StylusInput/Source/StylusInput/Private/WindowsStylusInputInterface.h b/Engine/Plugins/Editor/StylusInput/Source/StylusInput/Private/WindowsStylusInputInterface.h
--- a/Engine/Plugins/Editor/StylusInput/Source/StylusInput/Private/WindowsStylusInputInterface.h
+++ b/Engine/Plugins/Editor/StylusInput/Source/StylusInput/Private/WindowsStylusInputInterface.h
@@ -1,5 +1,7 @@
 #pragma once
 
+#include "IStylusInputModule.h"
+
 class FWindowsStylusInputInterfaceImpl;
 
 class FWindowsStylusInputInterface : public IStylusInputInterfaceInternal
diff --git a/Engine/Plugins/Editor/StylusInput/Source/StylusInput/Public/IStylusInputModule.h b/Engine/Plugins/Editor/StylusInput/Source/StylusInput/Public/IStylusInputModule.h
--- a/Engine/Plugins/Editor/StylusInput/Source/StylusInput/Public/IStylusInputModule.h
+++ b/Engine/Plugins/Editor/StylusInput/Source/StylusInput/Public/IStylusInputModule.h
@@ -8,6 +8,9 @@
 #include "Modules/ModuleManager.h"
 #include "EditorSubsystem.h"
 
+#include "TickableEditorObject.h"
+#include "Widgets\Docking\SDockTab.h"
+
 #include "IStylusInputModule.generated.h"
 
 /**
@@ -36,7 +39,14 @@ public:
 	}
 };
 
-class IStylusInputInterfaceInternal;
+// This is the interface that all platform-specific implementations must implement.
+class IStylusInputInterfaceInternal
+{
+public:
+	virtual IStylusInputDevice* GetInputDevice(int32 Index) const = 0;
+	virtual int32 NumInputDevices() const = 0;
+};
+
 
 UCLASS()
 class STYLUSINPUT_API UStylusInputSubsystem : 
