Fixing invalid include format.

diff --git a/Engine/Plugins/Editor/StylusInput/Source/StylusInput/Public/IStylusInputModule.h b/Engine/Plugins/Editor/StylusInput/Source/StylusInput/Public/IStylusInputModule.h
--- a/Engine/Plugins/Editor/StylusInput/Source/StylusInput/Public/IStylusInputModule.h
+++ b/Engine/Plugins/Editor/StylusInput/Source/StylusInput/Public/IStylusInputModule.h
@@ -4,12 +4,11 @@
 
 #include "IStylusState.h"
 
-#include "Modules/ModuleInterface.h"
-#include "Modules/ModuleManager.h"
 #include "EditorSubsystem.h"
-
 #include "TickableEditorObject.h"
-#include "Widgets\Docking\SDockTab.h"
+#include "Modules/ModuleInterface.h"
+#include "Modules/ModuleManager.h"
+#include "Widgets/Docking/SDockTab.h"
 
 #include "IStylusInputModule.generated.h"
 
