Fixing Linux platform CIS build.

diff --git a/Engine/Plugins/Editor/StylusInput/Source/StylusInput/Private/StylusInputModule.cpp b/Engine/Plugins/Editor/StylusInput/Source/StylusInput/Private/StylusInputModule.cpp
--- a/Engine/Plugins/Editor/StylusInput/Source/StylusInput/Private/StylusInputModule.cpp
+++ b/Engine/Plugins/Editor/StylusInput/Source/StylusInput/Private/StylusInputModule.cpp
@@ -23,13 +23,14 @@ class FStylusInputModule : public IModuleInterface
 
 IMPLEMENT_MODULE(FStylusInputModule, StylusInput)
 
-
 // This is the function that all platform-specific implementations are required to implement.
 TSharedPtr<IStylusInputInterfaceInternal> CreateStylusInputInterface();
 
 #if PLATFORM_WINDOWS
 #include "WindowsStylusInputInterface.h"
-#endif // PLATFORM_WINDOWS
+#else
+TSharedPtr<IStylusInputInterfaceInternal> CreateStylusInputInterface() { return TSharedPtr<IStylusInputInterfaceInternal>(); }
+#endif
 
 // TODO: Other platforms
 
@@ -41,6 +42,12 @@ void UStylusInputSubsystem::Initialize(FSubsystemCollectionBase& Collection)
 
 	InputInterface = CreateStylusInputInterface();
 
+	if (!InputInterface.IsValid())
+	{
+		UE_LOG(LogStylusInput, Log, TEXT("StylusInput not supported on this platform."));
+		return;
+	}
+
 	const TSharedRef<FGlobalTabmanager>& TabManager = FGlobalTabmanager::Get();
 	const IWorkspaceMenuStructure& MenuStructure = WorkspaceMenu::GetMenuStructure();
 
diff --git a/Engine/Plugins/Editor/StylusInput/Source/StylusInput/Private/WindowsRealTimeStylusPlugin.cpp b/Engine/Plugins/Editor/StylusInput/Source/StylusInput/Private/WindowsRealTimeStylusPlugin.cpp
--- a/Engine/Plugins/Editor/StylusInput/Source/StylusInput/Private/WindowsRealTimeStylusPlugin.cpp
+++ b/Engine/Plugins/Editor/StylusInput/Source/StylusInput/Private/WindowsRealTimeStylusPlugin.cpp
@@ -2,6 +2,8 @@
 
 #include "WindowsRealTimeStylusPlugin.h"
 
+#if PLATFORM_WINDOWS
+
 HRESULT FWindowsRealTimeStylusPlugin::QueryInterface(const IID& InterfaceID, void** Pointer)
 {
 	if ((InterfaceID == __uuidof(IStylusSyncPlugin)) || (InterfaceID == IID_IUnknown))
@@ -387,4 +389,6 @@ HRESULT FWindowsRealTimeStylusPlugin::InAirPackets(IRealTimeStylus* RealTimeStyl
 {
 	HandlePacket(RealTimeStylus, StylusInfo, PacketCount, PacketBufferLength, Packets);
 	return S_OK;
-}
\ No newline at end of file
+}
+
+#endif // PLATFORM_WINDOWS
\ No newline at end of file
diff --git a/Engine/Plugins/Editor/StylusInput/Source/StylusInput/Private/WindowsRealTimeStylusPlugin.h b/Engine/Plugins/Editor/StylusInput/Source/StylusInput/Private/WindowsRealTimeStylusPlugin.h
--- a/Engine/Plugins/Editor/StylusInput/Source/StylusInput/Private/WindowsRealTimeStylusPlugin.h
+++ b/Engine/Plugins/Editor/StylusInput/Source/StylusInput/Private/WindowsRealTimeStylusPlugin.h
@@ -3,6 +3,8 @@
 #pragma once
 
 #include "CoreMinimal.h"
+
+#if PLATFORM_WINDOWS
 #include "Windows/WindowsHWrapper.h"
 
 #include "Windows/AllowWindowsPlatformTypes.h"
@@ -169,4 +171,6 @@ private:
 
 	void AddTabletContext(IRealTimeStylus* RealTimeStylus, TABLET_CONTEXT_ID TabletID);
 	void RemoveTabletContext(IRealTimeStylus* RealTimeStylus, TABLET_CONTEXT_ID TabletID);
-};
\ No newline at end of file
+};
+
+#endif // PLATFORM_WINDOWS
\ No newline at end of file
diff --git a/Engine/Plugins/Editor/StylusInput/Source/StylusInput/Private/WindowsStylusInputInterface.cpp b/Engine/Plugins/Editor/StylusInput/Source/StylusInput/Private/WindowsStylusInputInterface.cpp
--- a/Engine/Plugins/Editor/StylusInput/Source/StylusInput/Private/WindowsStylusInputInterface.cpp
+++ b/Engine/Plugins/Editor/StylusInput/Source/StylusInput/Private/WindowsStylusInputInterface.cpp
@@ -4,6 +4,8 @@
 #include "WindowsRealTimeStylusPlugin.h"
 #include "Interfaces/IMainFrameModule.h"
 
+#if PLATFORM_WINDOWS
+
 class FWindowsStylusInputInterfaceImpl
 {
 public:
@@ -157,3 +159,5 @@ TSharedPtr<IStylusInputInterfaceInternal> CreateStylusInputInterface()
 
 	return MakeShareable(new FWindowsStylusInputInterface(WindowsImpl));
 }
+
+#endif // PLATFORM_WINDOWS
\ No newline at end of file
