Fixing static analysis warnings - copyright boilerplate, memory leak in error conditions, and wchar_t* to BSTR conversion.

diff --git a/Engine/Plugins/Editor/StylusInput/Source/StylusInput/Private/WindowsRealTimeStylusPlugin.cpp b/Engine/Plugins/Editor/StylusInput/Source/StylusInput/Private/WindowsRealTimeStylusPlugin.cpp
--- a/Engine/Plugins/Editor/StylusInput/Source/StylusInput/Private/WindowsRealTimeStylusPlugin.cpp
+++ b/Engine/Plugins/Editor/StylusInput/Source/StylusInput/Private/WindowsRealTimeStylusPlugin.cpp
@@ -130,95 +130,141 @@ static void SetupTabletSupportedPackets(TComPtr<IRealTimeStylus> RealTimeStylus,
 	RealTimeStylus->GetTabletFromTabletContextId(TabletContext.ID, &InkTablet);
 
 	int16 Supported;
-	InkTablet->IsPacketPropertySupported(STR_GUID_X, &Supported);
+
+	BSTR GuidBSTR;
+	
+	GuidBSTR = SysAllocString(STR_GUID_X);
+
+	InkTablet->IsPacketPropertySupported(GuidBSTR, &Supported);
 	if (Supported)
 	{
 		TabletContext.SupportedPackets.Add(EWindowsPacketType::X);
 	}
 
-	InkTablet->IsPacketPropertySupported(STR_GUID_Y, &Supported);
+	SysFreeString(GuidBSTR);
+	GuidBSTR = SysAllocString(STR_GUID_Y);
+
+	InkTablet->IsPacketPropertySupported(GuidBSTR, &Supported);
 	if (Supported)
 	{
 		TabletContext.SupportedPackets.Add(EWindowsPacketType::Y);
 	}
+	
+	SysFreeString(GuidBSTR);
+	GuidBSTR = SysAllocString(STR_GUID_Z);
 
-	InkTablet->IsPacketPropertySupported(STR_GUID_Z, &Supported);
+	InkTablet->IsPacketPropertySupported(GuidBSTR, &Supported);
 	if (Supported)
 	{
 		TabletContext.SupportedPackets.Add(EWindowsPacketType::Z);
 		TabletContext.AddSupportedInput(EStylusInputType::Z);
 	}
 
-	InkTablet->IsPacketPropertySupported(STR_GUID_PAKETSTATUS, &Supported);
+	SysFreeString(GuidBSTR);
+	GuidBSTR = SysAllocString(STR_GUID_PAKETSTATUS);
+	
+	InkTablet->IsPacketPropertySupported(GuidBSTR, &Supported);
 	if (Supported)
 	{
 		TabletContext.SupportedPackets.Add(EWindowsPacketType::Status);
 	}
 
-	InkTablet->IsPacketPropertySupported(STR_GUID_NORMALPRESSURE, &Supported);
+	SysFreeString(GuidBSTR);
+	GuidBSTR = SysAllocString(STR_GUID_NORMALPRESSURE);
+
+	InkTablet->IsPacketPropertySupported(GuidBSTR, &Supported);
 	if (Supported)
 	{
 		TabletContext.SupportedPackets.Add(EWindowsPacketType::NormalPressure);
 		TabletContext.AddSupportedInput(EStylusInputType::Pressure);
 	}
 
-	InkTablet->IsPacketPropertySupported(STR_GUID_TANGENTPRESSURE, &Supported);
+	SysFreeString(GuidBSTR);
+	GuidBSTR = SysAllocString(STR_GUID_TANGENTPRESSURE);
+
+	InkTablet->IsPacketPropertySupported(GuidBSTR, &Supported);
 	if (Supported)
 	{
 		TabletContext.SupportedPackets.Add(EWindowsPacketType::TangentPressure);
 		TabletContext.AddSupportedInput(EStylusInputType::TangentPressure);
 	}
 
-	InkTablet->IsPacketPropertySupported(STR_GUID_BUTTONPRESSURE, &Supported);
+	SysFreeString(GuidBSTR);
+	GuidBSTR = SysAllocString(STR_GUID_BUTTONPRESSURE);
+
+	InkTablet->IsPacketPropertySupported(GuidBSTR, &Supported);
 	if (Supported)
 	{
 		TabletContext.SupportedPackets.Add(EWindowsPacketType::ButtonPressure);
 		TabletContext.AddSupportedInput(EStylusInputType::ButtonPressure);
 	}
 
-	InkTablet->IsPacketPropertySupported(STR_GUID_AZIMUTHORIENTATION, &Supported);
+	SysFreeString(GuidBSTR);
+	GuidBSTR = SysAllocString(STR_GUID_AZIMUTHORIENTATION);
+
+	InkTablet->IsPacketPropertySupported(GuidBSTR, &Supported);
 	if (Supported)
 	{
 		TabletContext.SupportedPackets.Add(EWindowsPacketType::Azimuth);
 	}
 
-	InkTablet->IsPacketPropertySupported(STR_GUID_ALTITUDEORIENTATION, &Supported);
+	SysFreeString(GuidBSTR);
+	GuidBSTR = SysAllocString(STR_GUID_ALTITUDEORIENTATION);
+
+	InkTablet->IsPacketPropertySupported(GuidBSTR, &Supported);
 	if (Supported)
 	{
 		TabletContext.SupportedPackets.Add(EWindowsPacketType::Altitude);
 	}
 
-	InkTablet->IsPacketPropertySupported(STR_GUID_XTILTORIENTATION, &Supported);
+	SysFreeString(GuidBSTR);
+	GuidBSTR = SysAllocString(STR_GUID_XTILTORIENTATION);
+	
+	InkTablet->IsPacketPropertySupported(GuidBSTR, &Supported);
 	if (Supported)
 	{
 		TabletContext.SupportedPackets.Add(EWindowsPacketType::XTilt);
 	}
 
-	InkTablet->IsPacketPropertySupported(STR_GUID_YTILTORIENTATION, &Supported);
+	SysFreeString(GuidBSTR);
+	GuidBSTR = SysAllocString(STR_GUID_YTILTORIENTATION);
+
+	InkTablet->IsPacketPropertySupported(GuidBSTR, &Supported);
 	if (Supported)
 	{
 		TabletContext.SupportedPackets.Add(EWindowsPacketType::YTilt);
 	}
 
-	InkTablet->IsPacketPropertySupported(STR_GUID_TWISTORIENTATION, &Supported);
+	SysFreeString(GuidBSTR);
+	GuidBSTR = SysAllocString(STR_GUID_TWISTORIENTATION);
+
+	InkTablet->IsPacketPropertySupported(GuidBSTR, &Supported);
 	if (Supported)
 	{
 		TabletContext.SupportedPackets.Add(EWindowsPacketType::Twist);
 		TabletContext.AddSupportedInput(EStylusInputType::Twist);
 	}
 
-	InkTablet->IsPacketPropertySupported(STR_GUID_WIDTH, &Supported);
+	SysFreeString(GuidBSTR);
+	GuidBSTR = SysAllocString(STR_GUID_WIDTH);
+
+	InkTablet->IsPacketPropertySupported(GuidBSTR, &Supported);
 	if (Supported)
 	{
 		TabletContext.SupportedPackets.Add(EWindowsPacketType::Width);
 	}
 
-	InkTablet->IsPacketPropertySupported(STR_GUID_HEIGHT, &Supported);
+	SysFreeString(GuidBSTR);
+	GuidBSTR = SysAllocString(STR_GUID_HEIGHT);
+
+	InkTablet->IsPacketPropertySupported(GuidBSTR, &Supported);
 	if (Supported)
 	{
 		TabletContext.SupportedPackets.Add(EWindowsPacketType::Height);
 	}
 
+	SysFreeString(GuidBSTR);
+
 	if (TabletContext.SupportedPackets.Contains(EWindowsPacketType::X) &&
 		TabletContext.SupportedPackets.Contains(EWindowsPacketType::Y))
 	{
diff --git a/Engine/Plugins/Editor/StylusInput/Source/StylusInput/Private/WindowsStylusInputInterface.cpp b/Engine/Plugins/Editor/StylusInput/Source/StylusInput/Private/WindowsStylusInputInterface.cpp
--- a/Engine/Plugins/Editor/StylusInput/Source/StylusInput/Private/WindowsStylusInputInterface.cpp
+++ b/Engine/Plugins/Editor/StylusInput/Source/StylusInput/Private/WindowsStylusInputInterface.cpp
@@ -89,11 +89,12 @@ static void OnMainFrameCreated(FWindowsStylusInputInterfaceImpl& WindowsImpl, TS
 TSharedPtr<IStylusInputInterfaceInternal> CreateStylusInputInterface()
 {
 	FWindowsStylusInputInterfaceImpl* WindowsImpl = new FWindowsStylusInputInterfaceImpl();
+	TSharedPtr<IStylusInputInterfaceInternal> InterfaceImpl = MakeShareable(new FWindowsStylusInputInterface(WindowsImpl));
 
 	if (!FWindowsPlatformMisc::CoInitialize()) 
 	{
 		UE_LOG(LogStylusInput, Error, TEXT("Could not initialize COM library!"));
-		return TSharedPtr<IStylusInputInterfaceInternal>();
+		return nullptr;
 	}
 
 	// Load RealTimeStylus DLL
@@ -106,7 +107,7 @@ TSharedPtr<IStylusInputInterfaceInternal> CreateStylusInputInterface()
 	{
 		FWindowsPlatformMisc::CoUninitialize();
 		UE_LOG(LogStylusInput, Error, TEXT("Could not load RTSCom.dll!"));
-		return TSharedPtr<IStylusInputInterfaceInternal>();
+		return nullptr;
 	}
 
 	FPlatformProcess::PopDllDirectory(*InkDLLDirectory);
@@ -118,7 +119,7 @@ TSharedPtr<IStylusInputInterfaceInternal> CreateStylusInputInterface()
 	{
 		FWindowsPlatformMisc::CoUninitialize();
 		UE_LOG(LogStylusInput, Error, TEXT("Could not create RealTimeStylus!"));
-		return TSharedPtr<IStylusInputInterfaceInternal>();
+		return nullptr;
 	}
 
 	WindowsImpl->RealTimeStylus = static_cast<IRealTimeStylus*>(OutInstance);
@@ -131,7 +132,7 @@ TSharedPtr<IStylusInputInterfaceInternal> CreateStylusInputInterface()
 	{
 		FWindowsPlatformMisc::CoUninitialize();
 		UE_LOG(LogStylusInput, Error, TEXT("Could not create FreeThreadedMarshaller!"));
-		return TSharedPtr<IStylusInputInterfaceInternal>();
+		return nullptr;
 	}
 
 	// Add stylus plugin to the interface
@@ -140,7 +141,7 @@ TSharedPtr<IStylusInputInterfaceInternal> CreateStylusInputInterface()
 	{
 		FWindowsPlatformMisc::CoUninitialize();
 		UE_LOG(LogStylusInput, Error, TEXT("Could not add stylus plugin to API!"));
-		return TSharedPtr<IStylusInputInterfaceInternal>();
+		return nullptr;
 	}
 
 	// Set hook to catch main window creation
@@ -157,7 +158,7 @@ TSharedPtr<IStylusInputInterfaceInternal> CreateStylusInputInterface()
 		OnMainFrameCreated(*WindowsImpl, MainFrameModule.GetParentWindow());
 	}
 
-	return MakeShareable(new FWindowsStylusInputInterface(WindowsImpl));
+	return InterfaceImpl;
 }
 
 #endif // PLATFORM_WINDOWS
\ No newline at end of file
diff --git a/Engine/Plugins/Editor/StylusInput/Source/StylusInput/Private/WindowsStylusInputInterface.h b/Engine/Plugins/Editor/StylusInput/Source/StylusInput/Private/WindowsStylusInputInterface.h
--- a/Engine/Plugins/Editor/StylusInput/Source/StylusInput/Private/WindowsStylusInputInterface.h
+++ b/Engine/Plugins/Editor/StylusInput/Source/StylusInput/Private/WindowsStylusInputInterface.h
@@ -1,3 +1,5 @@
+// Copyright 1998-2019 Epic Games, Inc. All Rights Reserved.
+
 #pragma once
 
 #include "IStylusInputModule.h"
