Fixed static analysis warnings about uninitialized variables, despite them being set from a Windows API.

diff --git a/Engine/Plugins/Editor/StylusInput/Source/StylusInput/Private/WindowsRealTimeStylusPlugin.cpp b/Engine/Plugins/Editor/StylusInput/Source/StylusInput/Private/WindowsRealTimeStylusPlugin.cpp
--- a/Engine/Plugins/Editor/StylusInput/Source/StylusInput/Private/WindowsRealTimeStylusPlugin.cpp
+++ b/Engine/Plugins/Editor/StylusInput/Source/StylusInput/Private/WindowsRealTimeStylusPlugin.cpp
@@ -45,10 +45,10 @@ HRESULT FWindowsRealTimeStylusPlugin::StylusUp(IRealTimeStylus* RealTimeStylus,
 
 static void SetupPacketDescriptions(IRealTimeStylus* RealTimeStylus, FTabletContextInfo& TabletContext)
 {
-	ULONG NumPacketProperties;
-	PACKET_PROPERTY* PacketProperties;
+	ULONG NumPacketProperties = 0;
+	PACKET_PROPERTY* PacketProperties = nullptr;
 	HRESULT hr = RealTimeStylus->GetPacketDescriptionData(TabletContext.ID, nullptr, nullptr, &NumPacketProperties, &PacketProperties);
-	if (SUCCEEDED(hr))
+	if (SUCCEEDED(hr) && PacketProperties != nullptr)
 	{
 		for (ULONG PropIdx = 0; PropIdx < NumPacketProperties; ++PropIdx)
 		{
