Fixed off-by-one error when normalizing stylus pressure.

diff --git a/Engine/Plugins/Editor/StylusInput/Source/StylusInput/Private/WindowsRealTimeStylusPlugin.cpp b/Engine/Plugins/Editor/StylusInput/Source/StylusInput/Private/WindowsRealTimeStylusPlugin.cpp
--- a/Engine/Plugins/Editor/StylusInput/Source/StylusInput/Private/WindowsRealTimeStylusPlugin.cpp
+++ b/Engine/Plugins/Editor/StylusInput/Source/StylusInput/Private/WindowsRealTimeStylusPlugin.cpp
@@ -358,7 +358,7 @@ HRESULT FWindowsRealTimeStylusPlugin::TabletRemoved(IRealTimeStylus* RealTimeSty
 
 static float Normalize(int Value, const FPacketDescription& Desc)
 {
-	return (float) (Value - Desc.Minimum) / (float) ((Desc.Maximum - Desc.Minimum) + 1);
+	return (float) (Value - Desc.Minimum) / (float) (Desc.Maximum - Desc.Minimum);
 }
 
 static float ToDegrees(int Value, const FPacketDescription& Desc)
