Stylus input plugin now uses whichever window is currently hovered as the source for stylus input events. This solves the issue where a popup window won't cause stylus input events to be triggered.

diff --git a/Engine/Plugins/Editor/StylusInput/Source/StylusInput/Private/StylusInputModule.cpp b/Engine/Plugins/Editor/StylusInput/Source/StylusInput/Private/StylusInputModule.cpp
--- a/Engine/Plugins/Editor/StylusInput/Source/StylusInput/Private/StylusInputModule.cpp
+++ b/Engine/Plugins/Editor/StylusInput/Source/StylusInput/Private/StylusInputModule.cpp
@@ -101,6 +101,8 @@ void UStylusInputSubsystem::Tick(float DeltaTime)
 {
 	if (InputInterface.IsValid())
 	{
+		InputInterface->Tick();
+
 		for (int32 DeviceIdx = 0; DeviceIdx < NumInputDevices(); ++DeviceIdx)
 		{
 			IStylusInputDevice* InputDevice = InputInterface->GetInputDevice(DeviceIdx);
diff --git a/Engine/Plugins/Editor/StylusInput/Source/StylusInput/Private/WindowsStylusInputInterface.cpp b/Engine/Plugins/Editor/StylusInput/Source/StylusInput/Private/WindowsStylusInputInterface.cpp
--- a/Engine/Plugins/Editor/StylusInput/Source/StylusInput/Private/WindowsStylusInputInterface.cpp
+++ b/Engine/Plugins/Editor/StylusInput/Source/StylusInput/Private/WindowsStylusInputInterface.cpp
@@ -4,6 +4,8 @@
 #include "WindowsRealTimeStylusPlugin.h"
 #include "Interfaces/IMainFrameModule.h"
 
+#include "Framework/Application/SlateApplication.h"
+
 #if PLATFORM_WINDOWS
 
 class FWindowsStylusInputInterfaceImpl
@@ -18,7 +20,30 @@ public:
 
 FWindowsStylusInputInterface::FWindowsStylusInputInterface(FWindowsStylusInputInterfaceImpl* InImpl)
 {
+	ensure(InImpl != nullptr);
+
 	Impl = InImpl;
+
+	// We desire to receive everything, but what we actually will receive is determined in AddTabletContext
+	TArray<GUID> DesiredPackets = {
+		GUID_PACKETPROPERTY_GUID_X,
+		GUID_PACKETPROPERTY_GUID_Y,
+		GUID_PACKETPROPERTY_GUID_Z,
+		GUID_PACKETPROPERTY_GUID_PACKET_STATUS,
+		GUID_PACKETPROPERTY_GUID_NORMAL_PRESSURE,
+		GUID_PACKETPROPERTY_GUID_TANGENT_PRESSURE,
+		GUID_PACKETPROPERTY_GUID_X_TILT_ORIENTATION,
+		GUID_PACKETPROPERTY_GUID_Y_TILT_ORIENTATION,
+		GUID_PACKETPROPERTY_GUID_TWIST_ORIENTATION,
+		GUID_PACKETPROPERTY_GUID_WIDTH,
+		GUID_PACKETPROPERTY_GUID_HEIGHT,
+		// Currently not needed.
+		//GUID_PACKETPROPERTY_GUID_BUTTON_PRESSURE,
+		//GUID_PACKETPROPERTY_GUID_AZIMUTH_ORIENTATION,
+		//GUID_PACKETPROPERTY_GUID_ALTITUDE_ORIENTATION,
+	};
+
+	Impl->RealTimeStylus->SetDesiredPacketDescription(DesiredPackets.Num(), DesiredPackets.GetData());
 }
 
 FWindowsStylusInputInterface::~FWindowsStylusInputInterface()
@@ -26,6 +51,33 @@ FWindowsStylusInputInterface::~FWindowsStylusInputInterface()
 	delete Impl;
 }
 
+void FWindowsStylusInputInterface::Tick()
+{
+	HANDLE_PTR HCurrentWnd;
+	Impl->RealTimeStylus->get_HWND(&HCurrentWnd);
+
+	FSlateApplication& Application = FSlateApplication::Get();
+
+	FWidgetPath WidgetPath = Application.LocateWindowUnderMouse(Application.GetCursorPos(), Application.GetInteractiveTopLevelWindows());
+	if (WidgetPath.IsValid())
+	{
+		TSharedPtr<SWindow> Window = WidgetPath.GetWindow();
+		if (Window.IsValid())
+		{
+			TSharedPtr<FGenericWindow> NativeWindow = Window->GetNativeWindow();
+			HWND Hwnd = reinterpret_cast<HWND>(NativeWindow->GetOSWindowHandle());
+
+			if (reinterpret_cast<HWND>(HCurrentWnd) != Hwnd)
+			{
+				// changing the HWND isn't supported when the plugin is enabled
+				Impl->RealTimeStylus->put_Enabled(Windows::FALSE); 
+				Impl->RealTimeStylus->put_HWND(reinterpret_cast<uint64>(Hwnd));
+				Impl->RealTimeStylus->put_Enabled(Windows::TRUE);
+			}
+		}
+	}
+}
+
 int32 FWindowsStylusInputInterface::NumInputDevices() const
 {
 	return Impl->StylusPlugin->TabletContexts.Num();
@@ -55,41 +107,9 @@ FWindowsStylusInputInterfaceImpl::~FWindowsStylusInputInterfaceImpl()
 	}
 }
 
-static void OnMainFrameCreated(FWindowsStylusInputInterfaceImpl& WindowsImpl, TSharedPtr<SWindow> Window)
-{
-	TSharedPtr<FGenericWindow> NativeWindow = Window->GetNativeWindow();
-	HWND Hwnd = reinterpret_cast<HWND>(NativeWindow->GetOSWindowHandle());
-
-	WindowsImpl.RealTimeStylus->put_HWND(reinterpret_cast<uint64>(Hwnd));
-
-	// We desire to receive everything, but what we actually will receive is determined in AddTabletContext
-	TArray<GUID> DesiredPackets = {
-		GUID_PACKETPROPERTY_GUID_X,
-		GUID_PACKETPROPERTY_GUID_Y,
-		GUID_PACKETPROPERTY_GUID_Z,
-		GUID_PACKETPROPERTY_GUID_PACKET_STATUS,
-		GUID_PACKETPROPERTY_GUID_NORMAL_PRESSURE,
-		GUID_PACKETPROPERTY_GUID_TANGENT_PRESSURE,
-		GUID_PACKETPROPERTY_GUID_X_TILT_ORIENTATION,
-		GUID_PACKETPROPERTY_GUID_Y_TILT_ORIENTATION,
-		GUID_PACKETPROPERTY_GUID_TWIST_ORIENTATION,
-		GUID_PACKETPROPERTY_GUID_WIDTH,
-		GUID_PACKETPROPERTY_GUID_HEIGHT,
-		// Currently not needed.
-		//GUID_PACKETPROPERTY_GUID_BUTTON_PRESSURE,
-		//GUID_PACKETPROPERTY_GUID_AZIMUTH_ORIENTATION,
-		//GUID_PACKETPROPERTY_GUID_ALTITUDE_ORIENTATION,
-	};
-
-	WindowsImpl.RealTimeStylus->SetDesiredPacketDescription(DesiredPackets.Num(), DesiredPackets.GetData());
-
-	WindowsImpl.RealTimeStylus->put_Enabled(Windows::TRUE);
-}
-
 TSharedPtr<IStylusInputInterfaceInternal> CreateStylusInputInterface()
 {
 	FWindowsStylusInputInterfaceImpl* WindowsImpl = new FWindowsStylusInputInterfaceImpl();
-	TSharedPtr<IStylusInputInterfaceInternal> InterfaceImpl = MakeShareable(new FWindowsStylusInputInterface(WindowsImpl));
 
 	if (!FWindowsPlatformMisc::CoInitialize()) 
 	{
@@ -126,8 +146,7 @@ TSharedPtr<IStylusInputInterfaceInternal> CreateStylusInputInterface()
 	WindowsImpl->StylusPlugin = MakeShareable(new FWindowsRealTimeStylusPlugin());
 	
 	// Create free-threaded marshaller for the plugin
-	hr = ::CoCreateFreeThreadedMarshaler(WindowsImpl->StylusPlugin.Get(), 
-		&WindowsImpl->StylusPlugin->FreeThreadedMarshaller);
+	hr = ::CoCreateFreeThreadedMarshaler(WindowsImpl->StylusPlugin.Get(), &WindowsImpl->StylusPlugin->FreeThreadedMarshaller);
 	if (FAILED(hr))
 	{
 		FWindowsPlatformMisc::CoUninitialize();
@@ -144,21 +163,8 @@ TSharedPtr<IStylusInputInterfaceInternal> CreateStylusInputInterface()
 		return nullptr;
 	}
 
-	// Set hook to catch main window creation
-	IMainFrameModule& MainFrameModule = FModuleManager::LoadModuleChecked<IMainFrameModule>("MainFrame");
-	if (!MainFrameModule.IsWindowInitialized())
-	{
-		MainFrameModule.OnMainFrameCreationFinished().AddLambda([WindowsImpl](TSharedPtr<SWindow> Window, bool)
-		{
-			OnMainFrameCreated(*WindowsImpl, Window);
-		});
-	}
-	else
-	{
-		OnMainFrameCreated(*WindowsImpl, MainFrameModule.GetParentWindow());
-	}
-
-	return InterfaceImpl;
+	
+	return MakeShareable(new FWindowsStylusInputInterface(WindowsImpl));
 }
 
 #endif // PLATFORM_WINDOWS
\ No newline at end of file
diff --git a/Engine/Plugins/Editor/StylusInput/Source/StylusInput/Private/WindowsStylusInputInterface.h b/Engine/Plugins/Editor/StylusInput/Source/StylusInput/Private/WindowsStylusInputInterface.h
--- a/Engine/Plugins/Editor/StylusInput/Source/StylusInput/Private/WindowsStylusInputInterface.h
+++ b/Engine/Plugins/Editor/StylusInput/Source/StylusInput/Private/WindowsStylusInputInterface.h
@@ -12,6 +12,7 @@ public:
 	FWindowsStylusInputInterface(FWindowsStylusInputInterfaceImpl* InImpl);
 	virtual ~FWindowsStylusInputInterface();
 
+	virtual void Tick() override;
 	virtual int32 NumInputDevices() const override;
 	virtual IStylusInputDevice* GetInputDevice(int32 Index) const override;
 
diff --git a/Engine/Plugins/Editor/StylusInput/Source/StylusInput/Public/IStylusInputModule.h b/Engine/Plugins/Editor/StylusInput/Source/StylusInput/Public/IStylusInputModule.h
--- a/Engine/Plugins/Editor/StylusInput/Source/StylusInput/Public/IStylusInputModule.h
+++ b/Engine/Plugins/Editor/StylusInput/Source/StylusInput/Public/IStylusInputModule.h
@@ -44,6 +44,8 @@ public:
 class IStylusInputInterfaceInternal
 {
 public:
+	virtual void Tick() = 0;
+
 	virtual IStylusInputDevice* GetInputDevice(int32 Index) const = 0;
 	virtual int32 NumInputDevices() const = 0;
 };
