Fixed static analysis warning about memory leak in Stylus Input initialization.

diff --git a/Engine/Plugins/Editor/StylusInput/Source/StylusInput/Private/WindowsStylusInputInterface.cpp b/Engine/Plugins/Editor/StylusInput/Source/StylusInput/Private/WindowsStylusInputInterface.cpp
--- a/Engine/Plugins/Editor/StylusInput/Source/StylusInput/Private/WindowsStylusInputInterface.cpp
+++ b/Engine/Plugins/Editor/StylusInput/Source/StylusInput/Private/WindowsStylusInputInterface.cpp
@@ -109,14 +109,14 @@ FWindowsStylusInputInterfaceImpl::~FWindowsStylusInputInterfaceImpl()
 
 TSharedPtr<IStylusInputInterfaceInternal> CreateStylusInputInterface()
 {
-	FWindowsStylusInputInterfaceImpl* WindowsImpl = new FWindowsStylusInputInterfaceImpl();
-
 	if (!FWindowsPlatformMisc::CoInitialize()) 
 	{
 		UE_LOG(LogStylusInput, Error, TEXT("Could not initialize COM library!"));
 		return nullptr;
 	}
 
+	FWindowsStylusInputInterfaceImpl* WindowsImpl = new FWindowsStylusInputInterfaceImpl();
+
 	// Load RealTimeStylus DLL
 	const FString InkDLLDirectory = TEXT("C:\\Program Files\\Common Files\\microsoft shared\\ink");
 	const FString RTSComDLL = TEXT("RTSCom.dll");
@@ -125,6 +125,7 @@ TSharedPtr<IStylusInputInterfaceInternal> CreateStylusInputInterface()
 	WindowsImpl->DLLHandle = FPlatformProcess::GetDllHandle(*(InkDLLDirectory / RTSComDLL));
 	if (WindowsImpl->DLLHandle == nullptr)
 	{
+		delete WindowsImpl;
 		FWindowsPlatformMisc::CoUninitialize();
 		UE_LOG(LogStylusInput, Error, TEXT("Could not load RTSCom.dll!"));
 		return nullptr;
@@ -137,6 +138,7 @@ TSharedPtr<IStylusInputInterfaceInternal> CreateStylusInputInterface()
 	HRESULT hr = ::CoCreateInstance(__uuidof(RealTimeStylus), nullptr, CLSCTX_INPROC, __uuidof(IRealTimeStylus), &OutInstance);
 	if (FAILED(hr))
 	{
+		delete WindowsImpl;
 		FWindowsPlatformMisc::CoUninitialize();
 		UE_LOG(LogStylusInput, Error, TEXT("Could not create RealTimeStylus!"));
 		return nullptr;
@@ -149,6 +151,7 @@ TSharedPtr<IStylusInputInterfaceInternal> CreateStylusInputInterface()
 	hr = ::CoCreateFreeThreadedMarshaler(WindowsImpl->StylusPlugin.Get(), &WindowsImpl->StylusPlugin->FreeThreadedMarshaller);
 	if (FAILED(hr))
 	{
+		delete WindowsImpl;
 		FWindowsPlatformMisc::CoUninitialize();
 		UE_LOG(LogStylusInput, Error, TEXT("Could not create FreeThreadedMarshaller!"));
 		return nullptr;
@@ -158,11 +161,11 @@ TSharedPtr<IStylusInputInterfaceInternal> CreateStylusInputInterface()
 	hr = WindowsImpl->RealTimeStylus->AddStylusSyncPlugin(0, WindowsImpl->StylusPlugin.Get());
 	if (FAILED(hr))
 	{
+		delete WindowsImpl;
 		FWindowsPlatformMisc::CoUninitialize();
 		UE_LOG(LogStylusInput, Error, TEXT("Could not add stylus plugin to API!"));
 		return nullptr;
 	}
-
 	
 	return MakeShareable(new FWindowsStylusInputInterface(WindowsImpl));
 }
