Further static analysis warning fixes.
Changed the impl class to be held by a TUniquePtr for additional safety and ownership.
Corrected use of MakeShareable to MakeShared.
Changed verify() to check() in the constructor, as it is never correct for the Impl object to be nullptr

diff --git a/Engine/Plugins/Editor/StylusInput/Source/StylusInput/Private/WindowsStylusInputInterface.cpp b/Engine/Plugins/Editor/StylusInput/Source/StylusInput/Private/WindowsStylusInputInterface.cpp
--- a/Engine/Plugins/Editor/StylusInput/Source/StylusInput/Private/WindowsStylusInputInterface.cpp
+++ b/Engine/Plugins/Editor/StylusInput/Source/StylusInput/Private/WindowsStylusInputInterface.cpp
@@ -18,11 +18,11 @@ public:
 	void* DLLHandle { nullptr };
 };
 
-FWindowsStylusInputInterface::FWindowsStylusInputInterface(FWindowsStylusInputInterfaceImpl* InImpl)
+FWindowsStylusInputInterface::FWindowsStylusInputInterface(TUniquePtr<FWindowsStylusInputInterfaceImpl> InImpl)
 {
-	ensure(InImpl != nullptr);
+	check(InImpl.IsValid());
 
-	Impl = InImpl;
+	Impl = MoveTemp(InImpl);
 
 	// We desire to receive everything, but what we actually will receive is determined in AddTabletContext
 	TArray<GUID> DesiredPackets = {
@@ -46,10 +46,7 @@ FWindowsStylusInputInterface::FWindowsStylusInputInterface(FWindowsStylusInputIn
 	Impl->RealTimeStylus->SetDesiredPacketDescription(DesiredPackets.Num(), DesiredPackets.GetData());
 }
 
-FWindowsStylusInputInterface::~FWindowsStylusInputInterface()
-{
-	delete Impl;
-}
+FWindowsStylusInputInterface::~FWindowsStylusInputInterface() = default;
 
 void FWindowsStylusInputInterface::Tick()
 {
@@ -115,7 +112,7 @@ TSharedPtr<IStylusInputInterfaceInternal> CreateStylusInputInterface()
 		return nullptr;
 	}
 
-	FWindowsStylusInputInterfaceImpl* WindowsImpl = new FWindowsStylusInputInterfaceImpl();
+	TUniquePtr<FWindowsStylusInputInterfaceImpl> WindowsImpl = MakeUnique<FWindowsStylusInputInterfaceImpl>();
 
 	// Load RealTimeStylus DLL
 	const FString InkDLLDirectory = TEXT("C:\\Program Files\\Common Files\\microsoft shared\\ink");
@@ -125,7 +122,6 @@ TSharedPtr<IStylusInputInterfaceInternal> CreateStylusInputInterface()
 	WindowsImpl->DLLHandle = FPlatformProcess::GetDllHandle(*(InkDLLDirectory / RTSComDLL));
 	if (WindowsImpl->DLLHandle == nullptr)
 	{
-		delete WindowsImpl;
 		FWindowsPlatformMisc::CoUninitialize();
 		UE_LOG(LogStylusInput, Error, TEXT("Could not load RTSCom.dll!"));
 		return nullptr;
@@ -138,7 +134,6 @@ TSharedPtr<IStylusInputInterfaceInternal> CreateStylusInputInterface()
 	HRESULT hr = ::CoCreateInstance(__uuidof(RealTimeStylus), nullptr, CLSCTX_INPROC, __uuidof(IRealTimeStylus), &OutInstance);
 	if (FAILED(hr))
 	{
-		delete WindowsImpl;
 		FWindowsPlatformMisc::CoUninitialize();
 		UE_LOG(LogStylusInput, Error, TEXT("Could not create RealTimeStylus!"));
 		return nullptr;
@@ -151,7 +146,6 @@ TSharedPtr<IStylusInputInterfaceInternal> CreateStylusInputInterface()
 	hr = ::CoCreateFreeThreadedMarshaler(WindowsImpl->StylusPlugin.Get(), &WindowsImpl->StylusPlugin->FreeThreadedMarshaller);
 	if (FAILED(hr))
 	{
-		delete WindowsImpl;
 		FWindowsPlatformMisc::CoUninitialize();
 		UE_LOG(LogStylusInput, Error, TEXT("Could not create FreeThreadedMarshaller!"));
 		return nullptr;
@@ -161,13 +155,12 @@ TSharedPtr<IStylusInputInterfaceInternal> CreateStylusInputInterface()
 	hr = WindowsImpl->RealTimeStylus->AddStylusSyncPlugin(0, WindowsImpl->StylusPlugin.Get());
 	if (FAILED(hr))
 	{
-		delete WindowsImpl;
 		FWindowsPlatformMisc::CoUninitialize();
 		UE_LOG(LogStylusInput, Error, TEXT("Could not add stylus plugin to API!"));
 		return nullptr;
 	}
 	
-	return MakeShareable(new FWindowsStylusInputInterface(WindowsImpl));
+	return MakeShared<FWindowsStylusInputInterface>(MoveTemp(WindowsImpl));
 }
 
 #endif // PLATFORM_WINDOWS
\ No newline at end of file
diff --git a/Engine/Plugins/Editor/StylusInput/Source/StylusInput/Private/WindowsStylusInputInterface.h b/Engine/Plugins/Editor/StylusInput/Source/StylusInput/Private/WindowsStylusInputInterface.h
--- a/Engine/Plugins/Editor/StylusInput/Source/StylusInput/Private/WindowsStylusInputInterface.h
+++ b/Engine/Plugins/Editor/StylusInput/Source/StylusInput/Private/WindowsStylusInputInterface.h
@@ -9,7 +9,7 @@ class FWindowsStylusInputInterfaceImpl;
 class FWindowsStylusInputInterface : public IStylusInputInterfaceInternal
 {
 public:
-	FWindowsStylusInputInterface(FWindowsStylusInputInterfaceImpl* InImpl);
+	FWindowsStylusInputInterface(TUniquePtr<FWindowsStylusInputInterfaceImpl> InImpl);
 	virtual ~FWindowsStylusInputInterface();
 
 	virtual void Tick() override;
@@ -18,6 +18,6 @@ public:
 
 private:
 	// pImpl to avoid including Windows headers.
-	FWindowsStylusInputInterfaceImpl* Impl;
+	TUniquePtr<FWindowsStylusInputInterfaceImpl> Impl;
 	TArray<IStylusMessageHandler*> MessageHandlers;
 };
