UStylusInputSubsystem is now marked as not being thread-safe for destruction, since FTickableEditorObject's destructor has to run on the main thread.

diff --git a/Engine/Plugins/Editor/StylusInput/Source/StylusInput/Public/IStylusInputModule.h b/Engine/Plugins/Editor/StylusInput/Source/StylusInput/Public/IStylusInputModule.h
--- a/Engine/Plugins/Editor/StylusInput/Source/StylusInput/Public/IStylusInputModule.h
+++ b/Engine/Plugins/Editor/StylusInput/Source/StylusInput/Public/IStylusInputModule.h
@@ -83,4 +83,7 @@ private:
 	TArray<IStylusMessageHandler*> MessageHandlers;
 
 	TSharedRef<SDockTab> OnSpawnPluginTab(const FSpawnTabArgs& Args);
+
+	/** FTickableEditorObject must be destroyed on the main thread. */
+	virtual bool IsDestructionThreadSafe() const override { return false; }
 };
\ No newline at end of file
