Stylus input no longer gets changed if the stylus is currently down, which would cause the stroke to essentially get cancelled, even if switching back to that window.

diff --git a/Engine/Plugins/Editor/StylusInput/Source/StylusInput/Private/WindowsStylusInputInterface.cpp b/Engine/Plugins/Editor/StylusInput/Source/StylusInput/Private/WindowsStylusInputInterface.cpp
--- a/Engine/Plugins/Editor/StylusInput/Source/StylusInput/Private/WindowsStylusInputInterface.cpp
+++ b/Engine/Plugins/Editor/StylusInput/Source/StylusInput/Private/WindowsStylusInputInterface.cpp
@@ -50,6 +50,15 @@ FWindowsStylusInputInterface::~FWindowsStylusInputInterface() = default;
 
 void FWindowsStylusInputInterface::Tick()
 {
+	for (const FTabletContextInfo& Context : Impl->StylusPlugin->TabletContexts)
+	{
+		// don't change focus if the stylus is down
+		if (Context.GetCurrentState().IsStylusDown())
+		{
+			return;
+		}
+	}
+
 	HANDLE_PTR HCurrentWnd;
 	Impl->RealTimeStylus->get_HWND(&HCurrentWnd);
 
