Fix null pointer access if the Stylus system fails to initialize

#rnx

diff --git a/Engine/Plugins/Editor/StylusInput/Source/StylusInput/Private/WindowsStylusInputInterface.cpp b/Engine/Plugins/Editor/StylusInput/Source/StylusInput/Private/WindowsStylusInputInterface.cpp
--- a/Engine/Plugins/Editor/StylusInput/Source/StylusInput/Private/WindowsStylusInputInterface.cpp
+++ b/Engine/Plugins/Editor/StylusInput/Source/StylusInput/Private/WindowsStylusInputInterface.cpp
@@ -101,8 +101,11 @@ IStylusInputDevice* FWindowsStylusInputInterface::GetInputDevice(int32 Index) co
 
 FWindowsStylusInputInterfaceImpl::~FWindowsStylusInputInterfaceImpl()
 {
-	RealTimeStylus->RemoveAllStylusSyncPlugins();
-	RealTimeStylus.Reset();
+	if (RealTimeStylus)
+	{
+		RealTimeStylus->RemoveAllStylusSyncPlugins();
+		RealTimeStylus.Reset();
+	}
 
 	StylusPlugin.Reset();
 
