COMPtr doesn't have explicit operator bool

#rnx

diff --git a/Engine/Plugins/Editor/StylusInput/Source/StylusInput/Private/WindowsStylusInputInterface.cpp b/Engine/Plugins/Editor/StylusInput/Source/StylusInput/Private/WindowsStylusInputInterface.cpp
--- a/Engine/Plugins/Editor/StylusInput/Source/StylusInput/Private/WindowsStylusInputInterface.cpp
+++ b/Engine/Plugins/Editor/StylusInput/Source/StylusInput/Private/WindowsStylusInputInterface.cpp
@@ -101,7 +101,7 @@ IStylusInputDevice* FWindowsStylusInputInterface::GetInputDevice(int32 Index) co
 
 FWindowsStylusInputInterfaceImpl::~FWindowsStylusInputInterfaceImpl()
 {
-	if (RealTimeStylus)
+	if (RealTimeStylus.IsValid())
 	{
 		RealTimeStylus->RemoveAllStylusSyncPlugins();
 		RealTimeStylus.Reset();
