Downgraded some stylus init errors to warnings

#rnx

diff --git a/Engine/Plugins/Editor/StylusInput/Source/StylusInput/Private/WindowsStylusInputInterface.cpp b/Engine/Plugins/Editor/StylusInput/Source/StylusInput/Private/WindowsStylusInputInterface.cpp
--- a/Engine/Plugins/Editor/StylusInput/Source/StylusInput/Private/WindowsStylusInputInterface.cpp
+++ b/Engine/Plugins/Editor/StylusInput/Source/StylusInput/Private/WindowsStylusInputInterface.cpp
@@ -135,7 +135,7 @@ TSharedPtr<IStylusInputInterfaceInternal> CreateStylusInputInterface()
 	if (WindowsImpl->DLLHandle == nullptr)
 	{
 		FWindowsPlatformMisc::CoUninitialize();
-		UE_LOG(LogStylusInput, Error, TEXT("Could not load RTSCom.dll!"));
+		UE_LOG(LogStylusInput, Warning, TEXT("Could not load RTSCom.dll! Stylus input will be unavailable."));
 		return nullptr;
 	}
 
@@ -147,7 +147,7 @@ TSharedPtr<IStylusInputInterfaceInternal> CreateStylusInputInterface()
 	if (FAILED(hr))
 	{
 		FWindowsPlatformMisc::CoUninitialize();
-		UE_LOG(LogStylusInput, Error, TEXT("Could not create RealTimeStylus!"));
+		UE_LOG(LogStylusInput, Warning, TEXT("Could not create RealTimeStylus! Stylus input will be unavailable."));
 		return nullptr;
 	}
 
@@ -159,7 +159,7 @@ TSharedPtr<IStylusInputInterfaceInternal> CreateStylusInputInterface()
 	if (FAILED(hr))
 	{
 		FWindowsPlatformMisc::CoUninitialize();
-		UE_LOG(LogStylusInput, Error, TEXT("Could not create FreeThreadedMarshaller!"));
+		UE_LOG(LogStylusInput, Warning, TEXT("Could not create FreeThreadedMarshaller! Stylus input will be unavailable."));
 		return nullptr;
 	}
 
@@ -168,7 +168,7 @@ TSharedPtr<IStylusInputInterfaceInternal> CreateStylusInputInterface()
 	if (FAILED(hr))
 	{
 		FWindowsPlatformMisc::CoUninitialize();
-		UE_LOG(LogStylusInput, Error, TEXT("Could not add stylus plugin to API!"));
+		UE_LOG(LogStylusInput, Warning, TEXT("Could not add stylus plugin to API! Stylus input will be unavailable."));
 		return nullptr;
 	}
 	
