Stylus input plugin usage of Windows Ink API has been changed to prevent an occasional deadlock that occurred when moving a mouse or stylus between open windows in the editor.

Previously, we created a single IRealTimeStylusSyncPlugin per IRealTimeStylus instance and attempted to switch the HWND that it was listening to, but this caused a deadlock when switching the HWND. Now we instead create an IRealTimeStylus and SyncPlugin per HWND and remove it when the window is closed. This removes the need for thread synchronization, and thus the deadlocks that could result.

There is an IRealTimeStylusSynchronization interface available, but my experiments with it did not resolve the issue and the documentation doesn't clarify its usage.

diff --git a/Engine/Plugins/Editor/StylusInput/Source/StylusInput/Private/WindowsRealTimeStylusPlugin.cpp b/Engine/Plugins/Editor/StylusInput/Source/StylusInput/Private/WindowsRealTimeStylusPlugin.cpp
--- a/Engine/Plugins/Editor/StylusInput/Source/StylusInput/Private/WindowsRealTimeStylusPlugin.cpp
+++ b/Engine/Plugins/Editor/StylusInput/Source/StylusInput/Private/WindowsRealTimeStylusPlugin.cpp
@@ -12,16 +12,16 @@ HRESULT FWindowsRealTimeStylusPlugin::QueryInterface(const IID& InterfaceID, voi
 		AddRef();
 		return S_OK;
 	}
-	else if ((InterfaceID == IID_IMarshal) && (FreeThreadedMarshaller != nullptr))
+	else if ((InterfaceID == IID_IMarshal) && (FreeThreadedMarshaler != nullptr))
 	{
-		return FreeThreadedMarshaller->QueryInterface(InterfaceID, Pointer);
+		return FreeThreadedMarshaler->QueryInterface(InterfaceID, Pointer);
 	}
 
 	*Pointer = nullptr;
 	return E_NOINTERFACE;
 }
 
-HRESULT FWindowsRealTimeStylusPlugin::StylusDown(IRealTimeStylus* RealTimeStylus, const StylusInfo* StylusInfo, ULONG PacketSize, LONG* Packet, LONG** InOutPackets)
+HRESULT FWindowsRealTimeStylusPlugin::StylusDown(IRealTimeStylus* InRealTimeStylus, const StylusInfo* StylusInfo, ULONG PacketSize, LONG* Packet, LONG** InOutPackets)
 {
 	FTabletContextInfo* TabletContext = FindTabletContext(StylusInfo->tcid);
 	if (TabletContext != nullptr)
@@ -31,7 +31,7 @@ HRESULT FWindowsRealTimeStylusPlugin::StylusDown(IRealTimeStylus* RealTimeStylus
 	return S_OK;
 }
 
-HRESULT FWindowsRealTimeStylusPlugin::StylusUp(IRealTimeStylus* RealTimeStylus, const StylusInfo* StylusInfo, ULONG PacketSize, LONG* Packet, LONG** InOutPackets)
+HRESULT FWindowsRealTimeStylusPlugin::StylusUp(IRealTimeStylus* InRealTimeStylus, const StylusInfo* StylusInfo, ULONG PacketSize, LONG* Packet, LONG** InOutPackets)
 {
 	FTabletContextInfo* TabletContext = FindTabletContext(StylusInfo->tcid);
 	if (TabletContext != nullptr)
@@ -43,11 +43,11 @@ HRESULT FWindowsRealTimeStylusPlugin::StylusUp(IRealTimeStylus* RealTimeStylus,
 	return S_OK;
 }
 
-static void SetupPacketDescriptions(IRealTimeStylus* RealTimeStylus, FTabletContextInfo& TabletContext)
+static void SetupPacketDescriptions(IRealTimeStylus* InRealTimeStylus, FTabletContextInfo& TabletContext)
 {
 	ULONG NumPacketProperties = 0;
 	PACKET_PROPERTY* PacketProperties = nullptr;
-	HRESULT hr = RealTimeStylus->GetPacketDescriptionData(TabletContext.ID, nullptr, nullptr, &NumPacketProperties, &PacketProperties);
+	HRESULT hr = InRealTimeStylus->GetPacketDescriptionData(TabletContext.ID, nullptr, nullptr, &NumPacketProperties, &PacketProperties);
 	if (SUCCEEDED(hr) && PacketProperties != nullptr)
 	{
 		for (ULONG PropIdx = 0; PropIdx < NumPacketProperties; ++PropIdx)
@@ -124,10 +124,10 @@ static void SetupPacketDescriptions(IRealTimeStylus* RealTimeStylus, FTabletCont
 	}
 }
 
-static void SetupTabletSupportedPackets(TComPtr<IRealTimeStylus> RealTimeStylus, FTabletContextInfo& TabletContext)
+static void SetupTabletSupportedPackets(TComPtr<IRealTimeStylus> InRealTimeStylus, FTabletContextInfo& TabletContext)
 {
 	IInkTablet* InkTablet;
-	RealTimeStylus->GetTabletFromTabletContextId(TabletContext.ID, &InkTablet);
+	InRealTimeStylus->GetTabletFromTabletContextId(TabletContext.ID, &InkTablet);
 
 	int16 Supported;
 
@@ -286,7 +286,7 @@ static void SetupTabletSupportedPackets(TComPtr<IRealTimeStylus> RealTimeStylus,
 
 FTabletContextInfo* FWindowsRealTimeStylusPlugin::FindTabletContext(TABLET_CONTEXT_ID TabletID)
 {
-	for (FTabletContextInfo& TabletContext : TabletContexts)
+	for (FTabletContextInfo& TabletContext : *TabletContexts)
 	{
 		if (TabletContext.ID == TabletID)
 		{
@@ -296,63 +296,69 @@ FTabletContextInfo* FWindowsRealTimeStylusPlugin::FindTabletContext(TABLET_CONTE
 	return nullptr;
 }
 
-void FWindowsRealTimeStylusPlugin::AddTabletContext(IRealTimeStylus* RealTimeStylus, TABLET_CONTEXT_ID TabletID)
+void FWindowsRealTimeStylusPlugin::AddTabletContext(IRealTimeStylus* InRealTimeStylus, TABLET_CONTEXT_ID TabletID)
 {
 	FTabletContextInfo* FoundContext = FindTabletContext(TabletID);
 	if (FoundContext == nullptr)
 	{
-		int32 Created = TabletContexts.Emplace();
-		FoundContext = &TabletContexts[Created];
+		FoundContext = &TabletContexts->Emplace_GetRef();
 		FoundContext->ID = TabletID;
 	}
 
-	SetupTabletSupportedPackets(RealTimeStylus, *FoundContext);
-	SetupPacketDescriptions(RealTimeStylus, *FoundContext);
+	SetupTabletSupportedPackets(InRealTimeStylus, *FoundContext);
+	SetupPacketDescriptions(InRealTimeStylus, *FoundContext);
 }
 
-void FWindowsRealTimeStylusPlugin::RemoveTabletContext(IRealTimeStylus* RealTimeStylus, TABLET_CONTEXT_ID TabletID)
+void FWindowsRealTimeStylusPlugin::RemoveTabletContext(IRealTimeStylus* InRealTimeStylus, TABLET_CONTEXT_ID TabletID)
 {
-	for (int32 ExistingIdx = 0; ExistingIdx < TabletContexts.Num(); ++ExistingIdx)
+	for (int32 ExistingIdx = 0; ExistingIdx < TabletContexts->Num(); ++ExistingIdx)
 	{
-		if (TabletContexts[ExistingIdx].ID == TabletID)
+		if ((*TabletContexts)[ExistingIdx].ID == TabletID)
 		{
-			TabletContexts.RemoveAt(ExistingIdx);
+			TabletContexts->RemoveAt(ExistingIdx);
 			break;
 		}
 	}
 }
 
-HRESULT FWindowsRealTimeStylusPlugin::RealTimeStylusEnabled(IRealTimeStylus* RealTimeStylus, ULONG Num, const TABLET_CONTEXT_ID* InTabletContexts)
+HRESULT FWindowsRealTimeStylusPlugin::RealTimeStylusEnabled(IRealTimeStylus* InRealTimeStylus, ULONG Num, const TABLET_CONTEXT_ID* InTabletContexts)
 {
 	for (ULONG TabletIdx = 0; TabletIdx < Num; ++TabletIdx)
 	{
-		AddTabletContext(RealTimeStylus, InTabletContexts[TabletIdx]);
+		AddTabletContext(InRealTimeStylus, InTabletContexts[TabletIdx]);
 	}
+
 	return S_OK;
 }
 
-HRESULT FWindowsRealTimeStylusPlugin::RealTimeStylusDisabled(IRealTimeStylus* RealTimeStylus, ULONG Num, const TABLET_CONTEXT_ID* InTabletContexts)
+HRESULT FWindowsRealTimeStylusPlugin::RealTimeStylusDisabled(IRealTimeStylus* InRealTimeStylus, ULONG Num, const TABLET_CONTEXT_ID* InTabletContexts)
 {
 	for (ULONG TabletIdx = 0; TabletIdx < Num; ++TabletIdx)
 	{
-		RemoveTabletContext(RealTimeStylus, InTabletContexts[TabletIdx]);
+		RemoveTabletContext(InRealTimeStylus, InTabletContexts[TabletIdx]);
 	}
+
 	return S_OK;
 }
 
-HRESULT FWindowsRealTimeStylusPlugin::TabletAdded(IRealTimeStylus* RealTimeStylus, IInkTablet* InkTablet)
+HRESULT FWindowsRealTimeStylusPlugin::TabletAdded(IRealTimeStylus* InRealTimeStylus, IInkTablet* InkTablet)
 {
 	TABLET_CONTEXT_ID TabletID;
-	if (SUCCEEDED(RealTimeStylus->GetTabletContextIdFromTablet(InkTablet, &TabletID)))
+	if (SUCCEEDED(InRealTimeStylus->GetTabletContextIdFromTablet(InkTablet, &TabletID)))
 	{
-		AddTabletContext(RealTimeStylus, TabletID);
+		AddTabletContext(InRealTimeStylus, TabletID);
 	}
+
 	return S_OK;
 }
 
-HRESULT FWindowsRealTimeStylusPlugin::TabletRemoved(IRealTimeStylus* RealTimeStylus, LONG iTabletIndex)
+HRESULT FWindowsRealTimeStylusPlugin::TabletRemoved(IRealTimeStylus* InRealTimeStylus, LONG iTabletIndex)
 {
-	TabletContexts.RemoveAt(iTabletIndex);
+	ULONG TabletContextCount = 0;
+	TABLET_CONTEXT_ID* AllContexts = nullptr;
+	InRealTimeStylus->GetAllTabletContextIds(&TabletContextCount, &AllContexts);
+
+	RemoveTabletContext(InRealTimeStylus, AllContexts[iTabletIndex]);
 	return S_OK;
 }
 
@@ -366,7 +372,7 @@ static float ToDegrees(int Value, const FPacketDescription& Desc)
 	return Value / Desc.Resolution;
 }
 
-void FWindowsRealTimeStylusPlugin::HandlePacket(IRealTimeStylus* RealTimeStylus, const StylusInfo* StylusInfo, ULONG PacketCount, ULONG PacketBufferLength, LONG* Packets)
+void FWindowsRealTimeStylusPlugin::HandlePacket(IRealTimeStylus* InRealTimeStylus, const StylusInfo* StylusInfo, ULONG PacketCount, ULONG PacketBufferLength, LONG* Packets)
 {
 	FTabletContextInfo* TabletContext = FindTabletContext(StylusInfo->tcid);
 	if (TabletContext == nullptr)
@@ -423,17 +429,17 @@ void FWindowsRealTimeStylusPlugin::HandlePacket(IRealTimeStylus* RealTimeStylus,
 	}
 }
 
-HRESULT FWindowsRealTimeStylusPlugin::Packets(IRealTimeStylus* RealTimeStylus, const StylusInfo* StylusInfo,
+HRESULT FWindowsRealTimeStylusPlugin::Packets(IRealTimeStylus* InRealTimeStylus, const StylusInfo* StylusInfo,
 	ULONG PacketCount, ULONG PacketBufferLength, LONG* Packets, ULONG* InOutPackets, LONG** PtrInOutPackets)
 {
-	HandlePacket(RealTimeStylus, StylusInfo, PacketCount, PacketBufferLength, Packets);
+	HandlePacket(InRealTimeStylus, StylusInfo, PacketCount, PacketBufferLength, Packets);
 	return S_OK;
 }
 
-HRESULT FWindowsRealTimeStylusPlugin::InAirPackets(IRealTimeStylus* RealTimeStylus, const StylusInfo* StylusInfo,
+HRESULT FWindowsRealTimeStylusPlugin::InAirPackets(IRealTimeStylus* InRealTimeStylus, const StylusInfo* StylusInfo,
 	ULONG PacketCount, ULONG PacketBufferLength, LONG* Packets, ULONG* InOutPackets, LONG** PtrInOutPackets)
 {
-	HandlePacket(RealTimeStylus, StylusInfo, PacketCount, PacketBufferLength, Packets);
+	HandlePacket(InRealTimeStylus, StylusInfo, PacketCount, PacketBufferLength, Packets);
 	return S_OK;
 }
 
diff --git a/Engine/Plugins/Editor/StylusInput/Source/StylusInput/Private/WindowsRealTimeStylusPlugin.h b/Engine/Plugins/Editor/StylusInput/Source/StylusInput/Private/WindowsRealTimeStylusPlugin.h
--- a/Engine/Plugins/Editor/StylusInput/Source/StylusInput/Private/WindowsRealTimeStylusPlugin.h
+++ b/Engine/Plugins/Editor/StylusInput/Source/StylusInput/Private/WindowsRealTimeStylusPlugin.h
@@ -99,7 +99,7 @@ struct FTabletContextInfo : public IStylusInputDevice
 };
 
 /**
- * An implementation of an IStylusSyncPlugin for use with the RealTimeStylus API.
+ * An implementation of an IStylusAsyncPlugin for use with the InRealTimeStylus API.
  */
 class FWindowsRealTimeStylusPlugin : public IStylusSyncPlugin
 {
@@ -107,9 +107,9 @@ public:
 	FWindowsRealTimeStylusPlugin() = default;
 	virtual ~FWindowsRealTimeStylusPlugin()
 	{
-		if (FreeThreadedMarshaller != nullptr)
+		if (FreeThreadedMarshaler != nullptr)
 		{
-			FreeThreadedMarshaller->Release();
+			FreeThreadedMarshaler->Release();
 		}
 	}
 
@@ -125,30 +125,30 @@ public:
 
 	virtual HRESULT QueryInterface(const IID& InterfaceID, void** Pointer) override;
 
-	virtual HRESULT TabletAdded(IRealTimeStylus* RealTimeStylus, IInkTablet* InkTablet) override;
-	virtual HRESULT TabletRemoved(IRealTimeStylus* RealTimeStylus, LONG iTabletIndex) override;
+	virtual HRESULT TabletAdded(IRealTimeStylus* InRealTimeStylus, IInkTablet* InkTablet) override;
+	virtual HRESULT TabletRemoved(IRealTimeStylus* InRealTimeStylus, LONG iTabletIndex) override;
 
-	virtual HRESULT RealTimeStylusEnabled(IRealTimeStylus* RealTimeStylus, ULONG Num, const TABLET_CONTEXT_ID* InTabletContexts) override;
-	virtual HRESULT RealTimeStylusDisabled(IRealTimeStylus* RealTimeStylus, ULONG Num, const TABLET_CONTEXT_ID* InTabletContexts) override;
+	virtual HRESULT RealTimeStylusEnabled(IRealTimeStylus* InRealTimeStylus, ULONG Num, const TABLET_CONTEXT_ID* InTabletContexts) override;
+	virtual HRESULT RealTimeStylusDisabled(IRealTimeStylus* InRealTimeStylus, ULONG Num, const TABLET_CONTEXT_ID* InTabletContexts) override;
 
-	virtual HRESULT StylusInRange(IRealTimeStylus* RealTimeStylus, TABLET_CONTEXT_ID TabletContext, STYLUS_ID StylusID) override { return S_OK; }
-	virtual HRESULT StylusOutOfRange(IRealTimeStylus* RealTimeStylus, TABLET_CONTEXT_ID TabletContext, STYLUS_ID StylusID) override { return S_OK; }
+	virtual HRESULT StylusInRange(IRealTimeStylus* InRealTimeStylus, TABLET_CONTEXT_ID TabletContext, STYLUS_ID StylusID) override { return S_OK; }
+	virtual HRESULT StylusOutOfRange(IRealTimeStylus* InRealTimeStylus, TABLET_CONTEXT_ID TabletContext, STYLUS_ID StylusID) override { return S_OK; }
 
-	virtual HRESULT StylusDown(IRealTimeStylus* RealTimeStylus, const StylusInfo* StylusInfo, ULONG PacketSize, LONG* Packet, LONG** InOutPackets) override;
-	virtual HRESULT StylusUp(IRealTimeStylus* RealTimeStylus, const StylusInfo* StylusInfo, ULONG PacketSize, LONG* Packet, LONG** InOutPackets) override;
+	virtual HRESULT StylusDown(IRealTimeStylus* InRealTimeStylus, const StylusInfo* StylusInfo, ULONG PacketSize, LONG* Packet, LONG** InOutPackets) override;
+	virtual HRESULT StylusUp(IRealTimeStylus* InRealTimeStylus, const StylusInfo* StylusInfo, ULONG PacketSize, LONG* Packet, LONG** InOutPackets) override;
 
-	virtual HRESULT StylusButtonDown(IRealTimeStylus* RealTimeStylus, STYLUS_ID StylusID, const GUID* GUID, POINT* Position) override { return S_OK; }
-	virtual HRESULT StylusButtonUp(IRealTimeStylus* RealTimeStylus, STYLUS_ID StylusID, const GUID* GUID, POINT* Position) override { return S_OK; }
+	virtual HRESULT StylusButtonDown(IRealTimeStylus* InRealTimeStylus, STYLUS_ID StylusID, const GUID* GUID, POINT* Position) override { return S_OK; }
+	virtual HRESULT StylusButtonUp(IRealTimeStylus* InRealTimeStylus, STYLUS_ID StylusID, const GUID* GUID, POINT* Position) override { return S_OK; }
 
-	virtual HRESULT InAirPackets(IRealTimeStylus* RealTimeStylus, const StylusInfo* StylusInfo,
+	virtual HRESULT InAirPackets(IRealTimeStylus* InRealTimeStylus, const StylusInfo* StylusInfo,
 		ULONG PacketCount, ULONG PacketBufferLength, LONG* Packets, ULONG* NumOutPackets, LONG** PtrOutPackets) override;
-	virtual HRESULT Packets(IRealTimeStylus* RealTimeStylus, const StylusInfo* StylusInfo,
+	virtual HRESULT Packets(IRealTimeStylus* InRealTimeStylus, const StylusInfo* StylusInfo,
 		ULONG PacketCount, ULONG PacketBufferSize, LONG* Packets, ULONG* NumOutPackets, LONG** PtrOutPackets) override;
 
-	virtual HRESULT CustomStylusDataAdded(IRealTimeStylus* RealTimeStylus, const GUID* GUID, ULONG Data, const BYTE* ByteData) override { return S_OK; }
+	virtual HRESULT CustomStylusDataAdded(IRealTimeStylus* InRealTimeStylus, const GUID* GUID, ULONG Data, const BYTE* ByteData) override { return S_OK; }
 
-	virtual HRESULT SystemEvent(IRealTimeStylus* RealTimeStylus, TABLET_CONTEXT_ID TabletContext, STYLUS_ID StylusID, SYSTEM_EVENT EventType, SYSTEM_EVENT_DATA EventData) override { return S_OK; }
-	virtual HRESULT Error(IRealTimeStylus* RealTimeStylus, IStylusPlugin* Plugin, RealTimeStylusDataInterest DataInterest, HRESULT ErrorCode, LONG_PTR* Key) override { return S_OK; }
+	virtual HRESULT SystemEvent(IRealTimeStylus* InRealTimeStylus, TABLET_CONTEXT_ID TabletContext, STYLUS_ID StylusID, SYSTEM_EVENT EventType, SYSTEM_EVENT_DATA EventData) override { return S_OK; }
+	virtual HRESULT Error(IRealTimeStylus* InRealTimeStylus, IStylusPlugin* Plugin, RealTimeStylusDataInterest DataInterest, HRESULT ErrorCode, LONG_PTR* Key) override { return S_OK; }
 
 	virtual HRESULT DataInterest(RealTimeStylusDataInterest* OutDataInterest) override
 	{
@@ -156,21 +156,23 @@ public:
 		return S_OK;
 	}
 
-	virtual HRESULT UpdateMapping(IRealTimeStylus* RealTimeStylus) override { return S_OK; }
+	virtual HRESULT UpdateMapping(IRealTimeStylus* InRealTimeStylus) override { return S_OK; }
 
 	FTabletContextInfo* FindTabletContext(TABLET_CONTEXT_ID TabletID);
 
-	IUnknown* FreeThreadedMarshaller;
-	TArray<FTabletContextInfo> TabletContexts;
 	bool HasChanges { false };
 
+	TArray<FTabletContextInfo>* TabletContexts;
+	TComPtr<IRealTimeStylus> RealTimeStylus;
+	IUnknown* FreeThreadedMarshaler;
+
 private:
 	int RefCount { 1 };
 
-	void HandlePacket(IRealTimeStylus* RealTimeStylus, const StylusInfo* StylusInfo, ULONG PacketCount, ULONG PacketBufferLength, LONG* Packets);
+	void HandlePacket(IRealTimeStylus* InRealTimeStylus, const StylusInfo* StylusInfo, ULONG PacketCount, ULONG PacketBufferLength, LONG* Packets);
 
-	void AddTabletContext(IRealTimeStylus* RealTimeStylus, TABLET_CONTEXT_ID TabletID);
-	void RemoveTabletContext(IRealTimeStylus* RealTimeStylus, TABLET_CONTEXT_ID TabletID);
+	void AddTabletContext(IRealTimeStylus* InRealTimeStylus, TABLET_CONTEXT_ID TabletID);
+	void RemoveTabletContext(IRealTimeStylus* InRealTimeStylus, TABLET_CONTEXT_ID TabletID);
 };
 
 #endif // PLATFORM_WINDOWS
\ No newline at end of file
diff --git a/Engine/Plugins/Editor/StylusInput/Source/StylusInput/Private/WindowsStylusInputInterface.cpp b/Engine/Plugins/Editor/StylusInput/Source/StylusInput/Private/WindowsStylusInputInterface.cpp
--- a/Engine/Plugins/Editor/StylusInput/Source/StylusInput/Private/WindowsStylusInputInterface.cpp
+++ b/Engine/Plugins/Editor/StylusInput/Source/StylusInput/Private/WindowsStylusInputInterface.cpp
@@ -13,9 +13,28 @@ class FWindowsStylusInputInterfaceImpl
 public:
 	~FWindowsStylusInputInterfaceImpl();
 
-	TComPtr<IRealTimeStylus> RealTimeStylus;
-	TSharedPtr<FWindowsRealTimeStylusPlugin> StylusPlugin;
+	TMap<void*, TSharedPtr<FWindowsRealTimeStylusPlugin>> StylusPlugins;
 	void* DLLHandle { nullptr };
+	TArray<FTabletContextInfo> TabletContexts;
+};
+
+// We desire to receive everything, but what we actually will receive is determined in AddTabletContext
+static const TArray<GUID> DesiredPackets = {
+	GUID_PACKETPROPERTY_GUID_X,
+	GUID_PACKETPROPERTY_GUID_Y,
+	GUID_PACKETPROPERTY_GUID_Z,
+	GUID_PACKETPROPERTY_GUID_PACKET_STATUS,
+	GUID_PACKETPROPERTY_GUID_NORMAL_PRESSURE,
+	GUID_PACKETPROPERTY_GUID_TANGENT_PRESSURE,
+	GUID_PACKETPROPERTY_GUID_X_TILT_ORIENTATION,
+	GUID_PACKETPROPERTY_GUID_Y_TILT_ORIENTATION,
+	GUID_PACKETPROPERTY_GUID_TWIST_ORIENTATION,
+	GUID_PACKETPROPERTY_GUID_WIDTH,
+	GUID_PACKETPROPERTY_GUID_HEIGHT,
+	// Currently not needed.
+	//GUID_PACKETPROPERTY_GUID_BUTTON_PRESSURE,
+	//GUID_PACKETPROPERTY_GUID_AZIMUTH_ORIENTATION,
+	//GUID_PACKETPROPERTY_GUID_ALTITUDE_ORIENTATION,
 };
 
 FWindowsStylusInputInterface::FWindowsStylusInputInterface(TUniquePtr<FWindowsStylusInputInterfaceImpl> InImpl)
@@ -23,92 +42,124 @@ FWindowsStylusInputInterface::FWindowsStylusInputInterface(TUniquePtr<FWindowsSt
 	check(InImpl.IsValid());
 
 	Impl = MoveTemp(InImpl);
-
-	// We desire to receive everything, but what we actually will receive is determined in AddTabletContext
-	TArray<GUID> DesiredPackets = {
-		GUID_PACKETPROPERTY_GUID_X,
-		GUID_PACKETPROPERTY_GUID_Y,
-		GUID_PACKETPROPERTY_GUID_Z,
-		GUID_PACKETPROPERTY_GUID_PACKET_STATUS,
-		GUID_PACKETPROPERTY_GUID_NORMAL_PRESSURE,
-		GUID_PACKETPROPERTY_GUID_TANGENT_PRESSURE,
-		GUID_PACKETPROPERTY_GUID_X_TILT_ORIENTATION,
-		GUID_PACKETPROPERTY_GUID_Y_TILT_ORIENTATION,
-		GUID_PACKETPROPERTY_GUID_TWIST_ORIENTATION,
-		GUID_PACKETPROPERTY_GUID_WIDTH,
-		GUID_PACKETPROPERTY_GUID_HEIGHT,
-		// Currently not needed.
-		//GUID_PACKETPROPERTY_GUID_BUTTON_PRESSURE,
-		//GUID_PACKETPROPERTY_GUID_AZIMUTH_ORIENTATION,
-		//GUID_PACKETPROPERTY_GUID_ALTITUDE_ORIENTATION,
-	};
-
-	Impl->RealTimeStylus->SetDesiredPacketDescription(DesiredPackets.Num(), DesiredPackets.GetData());
 }
 
 FWindowsStylusInputInterface::~FWindowsStylusInputInterface() = default;
 
+void FWindowsStylusInputInterface::CreateStylusPluginForHWND(void* HwndPtr)
+{
+	if (Impl->StylusPlugins.Contains(HwndPtr))
+	{
+		return;
+	}
+
+	HWND Hwnd = reinterpret_cast<HWND>(HwndPtr);
+
+	// Create RealTimeStylus interface
+	void* OutInstance = nullptr;
+	HRESULT hr = ::CoCreateInstance(__uuidof(RealTimeStylus), nullptr, CLSCTX_INPROC, __uuidof(IRealTimeStylus), &OutInstance);
+	if (FAILED(hr))
+	{
+		UE_LOG(LogStylusInput, Warning, TEXT("Could not create RealTimeStylus!"));
+		return;
+	}
+
+	TSharedPtr<FWindowsRealTimeStylusPlugin> NewPlugin = Impl->StylusPlugins.Add(HwndPtr, MakeShareable(new FWindowsRealTimeStylusPlugin()));
+	NewPlugin->RealTimeStylus = static_cast<IRealTimeStylus*>(OutInstance);
+
+	// Create free-threaded marshaller for the plugin
+	hr = ::CoCreateFreeThreadedMarshaler(NewPlugin.Get(), &NewPlugin->FreeThreadedMarshaler);
+	if (FAILED(hr))
+	{
+		UE_LOG(LogStylusInput, Warning, TEXT("Could not create FreeThreadedMarshaler!"));
+		return;
+	}
+
+	NewPlugin->TabletContexts = &Impl->TabletContexts;
+
+	NewPlugin->RealTimeStylus->SetDesiredPacketDescription(DesiredPackets.Num(), DesiredPackets.GetData());
+
+	// Add stylus plugin to the interface
+	hr = NewPlugin->RealTimeStylus->AddStylusSyncPlugin(0, NewPlugin.Get());
+	if (FAILED(hr))
+	{
+		UE_LOG(LogStylusInput, Warning, TEXT("Could not add stylus plugin to API!"));
+		return;
+	}
+	
+	NewPlugin->RealTimeStylus->put_HWND(reinterpret_cast<uint64>(Hwnd));
+	NewPlugin->RealTimeStylus->put_Enabled(Windows::TRUE);
+}
+
+void FWindowsStylusInputInterface::RemovePluginForWindow(const TSharedRef<SWindow>& Window)
+{
+	void* Hwnd = Window->GetNativeWindow()->GetOSWindowHandle();
+	TSharedPtr<FWindowsRealTimeStylusPlugin>* Plugin = Impl->StylusPlugins.Find(Hwnd);
+	if (Plugin != nullptr)
+	{
+		(*Plugin)->RealTimeStylus->put_Enabled(Windows::FALSE);
+		(*Plugin)->RealTimeStylus->RemoveAllStylusSyncPlugins();
+
+		Impl->StylusPlugins.Remove(Hwnd);
+	}
+}
+
 void FWindowsStylusInputInterface::Tick()
 {
-	for (const FTabletContextInfo& Context : Impl->StylusPlugin->TabletContexts)
+	for (const FTabletContextInfo& Context : Impl->TabletContexts)
 	{
-		// don't change focus if the stylus is down
+		// don't change focus if any stylus is down
 		if (Context.GetCurrentState().IsStylusDown())
 		{
 			return;
 		}
 	}
 
-	HANDLE_PTR HCurrentWnd;
-	Impl->RealTimeStylus->get_HWND(&HCurrentWnd);
-
 	FSlateApplication& Application = FSlateApplication::Get();
 
 	FWidgetPath WidgetPath = Application.LocateWindowUnderMouse(Application.GetCursorPos(), Application.GetInteractiveTopLevelWindows());
 	if (WidgetPath.IsValid())
 	{
 		TSharedPtr<SWindow> Window = WidgetPath.GetWindow();
-		if (Window.IsValid())
+		if (Window.IsValid() && Window->IsRegularWindow())
 		{
 			TSharedPtr<FGenericWindow> NativeWindow = Window->GetNativeWindow();
-			HWND Hwnd = reinterpret_cast<HWND>(NativeWindow->GetOSWindowHandle());
+			void* Hwnd = NativeWindow->GetOSWindowHandle();
 
-			if (reinterpret_cast<HWND>(HCurrentWnd) != Hwnd)
+			TSharedPtr<FWindowsRealTimeStylusPlugin>* CurrentPlugin = Impl->StylusPlugins.Find(Hwnd);
+			if (CurrentPlugin == nullptr)
 			{
-				// changing the HWND isn't supported when the plugin is enabled
-				Impl->RealTimeStylus->put_Enabled(Windows::FALSE); 
-				Impl->RealTimeStylus->put_HWND(reinterpret_cast<uint64>(Hwnd));
-				Impl->RealTimeStylus->put_Enabled(Windows::TRUE);
+				CreateStylusPluginForHWND(Hwnd);
 			}
+
+			Window->GetOnWindowClosedEvent().AddRaw(this, &FWindowsStylusInputInterface::RemovePluginForWindow);
 		}
 	}
 }
 
 int32 FWindowsStylusInputInterface::NumInputDevices() const
 {
-	return Impl->StylusPlugin->TabletContexts.Num();
+	return Impl->TabletContexts.Num();
 }
 
 IStylusInputDevice* FWindowsStylusInputInterface::GetInputDevice(int32 Index) const
 {
-	if (Index < 0 || Index >= Impl->StylusPlugin->TabletContexts.Num())
+	IStylusInputDevice* Result = nullptr;
+	if (Index >= 0 && Index < Impl->TabletContexts.Num())
 	{
-		return nullptr;
+		Result = &Impl->TabletContexts[Index];
 	}
-
-	return &Impl->StylusPlugin->TabletContexts[Index];
+	return Result;
 }
 
 FWindowsStylusInputInterfaceImpl::~FWindowsStylusInputInterfaceImpl()
 {
-	if (RealTimeStylus.IsValid())
+	for (const auto& Plugin : StylusPlugins)
 	{
-		RealTimeStylus->RemoveAllStylusSyncPlugins();
-		RealTimeStylus.Reset();
+		Plugin.Value->RealTimeStylus->RemoveAllStylusSyncPlugins();
+		Plugin.Value->RealTimeStylus.Reset();
 	}
 
-	StylusPlugin.Reset();
-
 	if (DLLHandle != nullptr)
 	{
 		FPlatformProcess::FreeDllHandle(DLLHandle);
@@ -135,43 +186,12 @@ TSharedPtr<IStylusInputInterfaceInternal> CreateStylusInputInterface()
 	if (WindowsImpl->DLLHandle == nullptr)
 	{
 		FWindowsPlatformMisc::CoUninitialize();
-		UE_LOG(LogStylusInput, Warning, TEXT("Could not load RTSCom.dll! Stylus input will be unavailable."));
+		UE_LOG(LogStylusInput, Error, TEXT("Could not load RTSCom.dll!"));
 		return nullptr;
 	}
 
 	FPlatformProcess::PopDllDirectory(*InkDLLDirectory);
 
-	// Create RealTimeStylus interface
-	void* OutInstance { nullptr };
-	HRESULT hr = ::CoCreateInstance(__uuidof(RealTimeStylus), nullptr, CLSCTX_INPROC, __uuidof(IRealTimeStylus), &OutInstance);
-	if (FAILED(hr))
-	{
-		FWindowsPlatformMisc::CoUninitialize();
-		UE_LOG(LogStylusInput, Warning, TEXT("Could not create RealTimeStylus! Stylus input will be unavailable."));
-		return nullptr;
-	}
-
-	WindowsImpl->RealTimeStylus = static_cast<IRealTimeStylus*>(OutInstance);
-	WindowsImpl->StylusPlugin = MakeShareable(new FWindowsRealTimeStylusPlugin());
-	
-	// Create free-threaded marshaller for the plugin
-	hr = ::CoCreateFreeThreadedMarshaler(WindowsImpl->StylusPlugin.Get(), &WindowsImpl->StylusPlugin->FreeThreadedMarshaller);
-	if (FAILED(hr))
-	{
-		FWindowsPlatformMisc::CoUninitialize();
-		UE_LOG(LogStylusInput, Warning, TEXT("Could not create FreeThreadedMarshaller! Stylus input will be unavailable."));
-		return nullptr;
-	}
-
-	// Add stylus plugin to the interface
-	hr = WindowsImpl->RealTimeStylus->AddStylusSyncPlugin(0, WindowsImpl->StylusPlugin.Get());
-	if (FAILED(hr))
-	{
-		FWindowsPlatformMisc::CoUninitialize();
-		UE_LOG(LogStylusInput, Warning, TEXT("Could not add stylus plugin to API! Stylus input will be unavailable."));
-		return nullptr;
-	}
-	
 	return MakeShared<FWindowsStylusInputInterface>(MoveTemp(WindowsImpl));
 }
 
diff --git a/Engine/Plugins/Editor/StylusInput/Source/StylusInput/Private/WindowsStylusInputInterface.h b/Engine/Plugins/Editor/StylusInput/Source/StylusInput/Private/WindowsStylusInputInterface.h
--- a/Engine/Plugins/Editor/StylusInput/Source/StylusInput/Private/WindowsStylusInputInterface.h
+++ b/Engine/Plugins/Editor/StylusInput/Source/StylusInput/Private/WindowsStylusInputInterface.h
@@ -5,6 +5,7 @@
 #include "IStylusInputModule.h"
 
 class FWindowsStylusInputInterfaceImpl;
+class SWindow;
 
 class FWindowsStylusInputInterface : public IStylusInputInterfaceInternal
 {
@@ -20,4 +21,7 @@ private:
 	// pImpl to avoid including Windows headers.
 	TUniquePtr<FWindowsStylusInputInterfaceImpl> Impl;
 	TArray<IStylusMessageHandler*> MessageHandlers;
+
+	void CreateStylusPluginForHWND(void* HwndPtr);
+	void RemovePluginForWindow(const TSharedRef<SWindow>& Window);
 };
