Downgrading stylus loading failure log to warning, because it was causing nightly builds to fail on CIS machines.

diff --git a/Engine/Plugins/Editor/StylusInput/Source/StylusInput/Private/WindowsStylusInputInterface.cpp b/Engine/Plugins/Editor/StylusInput/Source/StylusInput/Private/WindowsStylusInputInterface.cpp
--- a/Engine/Plugins/Editor/StylusInput/Source/StylusInput/Private/WindowsStylusInputInterface.cpp
+++ b/Engine/Plugins/Editor/StylusInput/Source/StylusInput/Private/WindowsStylusInputInterface.cpp
@@ -171,7 +171,7 @@ TSharedPtr<IStylusInputInterfaceInternal> CreateStylusInputInterface()
 {
 	if (!FWindowsPlatformMisc::CoInitialize()) 
 	{
-		UE_LOG(LogStylusInput, Error, TEXT("Could not initialize COM library!"));
+		UE_LOG(LogStylusInput, Warning, TEXT("Could not initialize COM library!"));
 		return nullptr;
 	}
 
@@ -186,7 +186,7 @@ TSharedPtr<IStylusInputInterfaceInternal> CreateStylusInputInterface()
 	if (WindowsImpl->DLLHandle == nullptr)
 	{
 		FWindowsPlatformMisc::CoUninitialize();
-		UE_LOG(LogStylusInput, Error, TEXT("Could not load RTSCom.dll!"));
+		UE_LOG(LogStylusInput, Warning, TEXT("Could not load RTSCom.dll!"));
 		return nullptr;
 	}
 
