Added error condition checks in the stylus plugin to silence static analysis.

diff --git a/Engine/Plugins/Editor/StylusInput/Source/StylusInput/Private/WindowsRealTimeStylusPlugin.cpp b/Engine/Plugins/Editor/StylusInput/Source/StylusInput/Private/WindowsRealTimeStylusPlugin.cpp
--- a/Engine/Plugins/Editor/StylusInput/Source/StylusInput/Private/WindowsRealTimeStylusPlugin.cpp
+++ b/Engine/Plugins/Editor/StylusInput/Source/StylusInput/Private/WindowsRealTimeStylusPlugin.cpp
@@ -357,7 +357,12 @@ HRESULT FWindowsRealTimeStylusPlugin::TabletRemoved(IRealTimeStylus* InRealTimeS
 	ULONG TabletContextCount = 0;
 	TABLET_CONTEXT_ID* AllContexts = nullptr;
 	InRealTimeStylus->GetAllTabletContextIds(&TabletContextCount, &AllContexts);
-
+	
+	if (AllContexts == nullptr || iTabletIndex < 0 || iTabletIndex >= (LONG) TabletContextCount)
+	{
+		return E_INVALIDARG;
+	}
+	
 	RemoveTabletContext(InRealTimeStylus, AllContexts[iTabletIndex]);
 	return S_OK;
 }
