Downgraded failure to load RealTimeStylus COM plugin to Display verbosity. Also finding the Common Files folder by using SHGetFolderPath.

#rnx

diff --git a/Engine/Plugins/Editor/StylusInput/Source/StylusInput/Private/WindowsStylusInputInterface.cpp b/Engine/Plugins/Editor/StylusInput/Source/StylusInput/Private/WindowsStylusInputInterface.cpp
--- a/Engine/Plugins/Editor/StylusInput/Source/StylusInput/Private/WindowsStylusInputInterface.cpp
+++ b/Engine/Plugins/Editor/StylusInput/Source/StylusInput/Private/WindowsStylusInputInterface.cpp
@@ -4,7 +4,9 @@
 #include "WindowsRealTimeStylusPlugin.h"
 #include "Interfaces/IMainFrameModule.h"
 
-#include "Framework/Application/SlateApplication.h"
+#include "Framework/Application/SlateApplication.h" 
+
+#include <ShlObj.h>
 
 #if PLATFORM_WINDOWS
 
@@ -177,16 +179,18 @@ TSharedPtr<IStylusInputInterfaceInternal> CreateStylusInputInterface()
 
 	TUniquePtr<FWindowsStylusInputInterfaceImpl> WindowsImpl = MakeUnique<FWindowsStylusInputInterfaceImpl>();
 
-	// Load RealTimeStylus DLL
-	const FString InkDLLDirectory = TEXT("C:\\Program Files\\Common Files\\microsoft shared\\ink");
-	const FString RTSComDLL = TEXT("RTSCom.dll");
+	// Load RealTimeStylus DLL 
+	TCHAR CommonFilesPath[MAX_PATH];
+	::SHGetFolderPath(0, CSIDL_PROGRAM_FILES_COMMON, NULL, 0, CommonFilesPath);
+	const FString InkDLLDirectory = FString(CommonFilesPath) + TEXT("\\microsoft shared\\ink");
 	FPlatformProcess::PushDllDirectory(*InkDLLDirectory);
 
+	const FString RTSComDLL = TEXT("RTSCom.dll");
 	WindowsImpl->DLLHandle = FPlatformProcess::GetDllHandle(*(InkDLLDirectory / RTSComDLL));
 	if (WindowsImpl->DLLHandle == nullptr)
 	{
 		FWindowsPlatformMisc::CoUninitialize();
-		UE_LOG(LogStylusInput, Warning, TEXT("Could not load RTSCom.dll!"));
+		UE_LOG(LogStylusInput, Display, TEXT("Could not load RTSCom.dll. Expected folder: %s"), *InkDLLDirectory);
 		return nullptr;
 	}
 
