Stylus input is no longer initialized when running commandlets.

Downgraded initialization warning in Stylus Input plugin to Display severity to match other initialization warnings in the plugin.

diff --git a/Engine/Plugins/Editor/StylusInput/Source/StylusInput/Private/StylusInputModule.cpp b/Engine/Plugins/Editor/StylusInput/Source/StylusInput/Private/StylusInputModule.cpp
--- a/Engine/Plugins/Editor/StylusInput/Source/StylusInput/Private/StylusInputModule.cpp
+++ b/Engine/Plugins/Editor/StylusInput/Source/StylusInput/Private/StylusInputModule.cpp
@@ -36,6 +36,11 @@ TSharedPtr<IStylusInputInterfaceInternal> CreateStylusInputInterface() { return
 
 void UStylusInputSubsystem::Initialize(FSubsystemCollectionBase& Collection)
 {
+	if (FApp::IsUnattended() || IsRunningCommandlet())
+	{
+		return;
+	}
+
 	Super::Initialize(Collection);
 
 	UE_LOG(LogStylusInput, Log, TEXT("Initializing StylusInput subsystem."));
diff --git a/Engine/Plugins/Editor/StylusInput/Source/StylusInput/Private/WindowsStylusInputInterface.cpp b/Engine/Plugins/Editor/StylusInput/Source/StylusInput/Private/WindowsStylusInputInterface.cpp
--- a/Engine/Plugins/Editor/StylusInput/Source/StylusInput/Private/WindowsStylusInputInterface.cpp
+++ b/Engine/Plugins/Editor/StylusInput/Source/StylusInput/Private/WindowsStylusInputInterface.cpp
@@ -177,7 +177,7 @@ TSharedPtr<IStylusInputInterfaceInternal> CreateStylusInputInterface()
 {
 	if (!FWindowsPlatformMisc::CoInitialize()) 
 	{
-		UE_LOG(LogStylusInput, Warning, TEXT("Could not initialize COM library!"));
+		UE_LOG(LogStylusInput, Display, TEXT("Could not initialize COM library!"));
 		return nullptr;
 	}
 
