Added missing App.h include causing CIS error in NoPCH Compile configuration.

diff --git a/Engine/Plugins/Editor/StylusInput/Source/StylusInput/Private/StylusInputModule.cpp b/Engine/Plugins/Editor/StylusInput/Source/StylusInput/Private/StylusInputModule.cpp
--- a/Engine/Plugins/Editor/StylusInput/Source/StylusInput/Private/StylusInputModule.cpp
+++ b/Engine/Plugins/Editor/StylusInput/Source/StylusInput/Private/StylusInputModule.cpp
@@ -6,6 +6,7 @@
 #include "Framework/Docking/TabManager.h"
 #include "Interfaces/IMainFrameModule.h"
 #include "Logging/LogMacros.h"
+#include "Misc/App.h"
 #include "Widgets/Docking/SDockTab.h"
 #include "Widgets/SWindow.h"
 #include "WorkspaceMenuStructure.h"
