[StylusInput]
* Attempt at fixing 2019 non unity strict compile error by re-adding includes (can't repro locally)

diff --git a/Engine/Plugins/Editor/StylusInput/Source/StylusInput/Private/WindowsStylusInputInterface.cpp b/Engine/Plugins/Editor/StylusInput/Source/StylusInput/Private/WindowsStylusInputInterface.cpp
--- a/Engine/Plugins/Editor/StylusInput/Source/StylusInput/Private/WindowsStylusInputInterface.cpp
+++ b/Engine/Plugins/Editor/StylusInput/Source/StylusInput/Private/WindowsStylusInputInterface.cpp
@@ -5,6 +5,10 @@
 
 #if PLATFORM_WINDOWS
 
+#include "WindowsRealTimeStylusPlugin.h"
+#include "Interfaces/IMainFrameModule.h"
+
+#include "Framework/Application/SlateApplication.h" 
 #include <ShlObj.h>
 
 class FWindowsStylusInputInterfaceImpl
