[StylusInput]
* Fixed weird bug with compile error in non unity where somehow WindowsRealTimeStylusPlugin.h removes MAX_int32 causing weird compile errors in UnrealString.h. But if UnrealString is pulled in before the weird code it all works since then unrealstring is already parsed into compile unit

diff --git a/Engine/Plugins/Editor/StylusInput/Source/StylusInput/Private/WindowsRealTimeStylusPlugin.h b/Engine/Plugins/Editor/StylusInput/Source/StylusInput/Private/WindowsRealTimeStylusPlugin.h
--- a/Engine/Plugins/Editor/StylusInput/Source/StylusInput/Private/WindowsRealTimeStylusPlugin.h
+++ b/Engine/Plugins/Editor/StylusInput/Source/StylusInput/Private/WindowsRealTimeStylusPlugin.h
@@ -2,7 +2,7 @@
 
 #pragma once
 
-#include "HAL/Platform.h"
+#include "Containers/UnrealString.h" // IWYU pragma: keep
 
 #if PLATFORM_WINDOWS
 #include "Windows/WindowsHWrapper.h"
diff --git a/Engine/Plugins/Editor/StylusInput/Source/StylusInput/Private/WindowsStylusInputInterface.cpp b/Engine/Plugins/Editor/StylusInput/Source/StylusInput/Private/WindowsStylusInputInterface.cpp
--- a/Engine/Plugins/Editor/StylusInput/Source/StylusInput/Private/WindowsStylusInputInterface.cpp
+++ b/Engine/Plugins/Editor/StylusInput/Source/StylusInput/Private/WindowsStylusInputInterface.cpp
@@ -5,10 +5,6 @@
 
 #if PLATFORM_WINDOWS
 
-#include "WindowsRealTimeStylusPlugin.h"
-#include "Interfaces/IMainFrameModule.h"
-
-#include "Framework/Application/SlateApplication.h" 
 #include <ShlObj.h>
 
 class FWindowsStylusInputInterfaceImpl
