Fixed mismatched copy constructor/assignment operator in FStylusState.

diff --git a/Engine/Plugins/Editor/StylusInput/Source/StylusInput/Public/IStylusState.h b/Engine/Plugins/Editor/StylusInput/Source/StylusInput/Public/IStylusState.h
--- a/Engine/Plugins/Editor/StylusInput/Source/StylusInput/Public/IStylusState.h
+++ b/Engine/Plugins/Editor/StylusInput/Source/StylusInput/Public/IStylusState.h
@@ -40,13 +40,6 @@ public:
 	{
 	}
 
-	FStylusState(const FStylusState& Other) :
-		Position(Other.Position), Z(Other.Z), Tilt(Other.Tilt), Twist(Other.Twist),
-		Pressure(Other.Pressure), TangentPressure(Other.TangentPressure), Size(Other.Size),
-		IsDown(Other.IsDown), IsInverted(Other.IsInverted)
-	{
-	}
-
 	/**
 	 * The current position of the stylus on (or above) the tablet. Always valid.
 	 * This value is in logical coordinates, not pixels.
