Stylus IsTouching should not be set to true when the mouse is down

diff --git a/Engine/Plugins/Editor/StylusInput/Source/StylusInput/Private/WindowsRealTimeStylusPlugin.cpp b/Engine/Plugins/Editor/StylusInput/Source/StylusInput/Private/WindowsRealTimeStylusPlugin.cpp
--- a/Engine/Plugins/Editor/StylusInput/Source/StylusInput/Private/WindowsRealTimeStylusPlugin.cpp
+++ b/Engine/Plugins/Editor/StylusInput/Source/StylusInput/Private/WindowsRealTimeStylusPlugin.cpp
@@ -24,7 +24,11 @@ HRESULT FWindowsRealTimeStylusPlugin::QueryInterface(const IID& InterfaceID, voi
 HRESULT FWindowsRealTimeStylusPlugin::StylusDown(IRealTimeStylus* InRealTimeStylus, const StylusInfo* StylusInfo, ULONG PacketSize, LONG* Packet, LONG** InOutPackets)
 {
 	FTabletContextInfo* TabletContext = FindTabletContext(StylusInfo->tcid);
-	if (TabletContext != nullptr)
+
+	// When mouse button is pressed StylusDown is called but we should not set IsTouching to true for mouse down
+	// Ideally it should filter the mouse down with TabletContext null check but some how mouse is also added as
+	// a tablet context, so to filter out mouse down adding cid check, cid for mouse is 1
+	if (TabletContext != nullptr && StylusInfo->cid != 1)
 	{
 		TabletContext->WindowsState.IsTouching = true;
 	}
