Fixed Slate accessing a deleted FWindowsStylusInputInterface when destroying windows after UStylusInputSubsystem had already shutdown

UStylusInputSubsystem would be destroyed during PreExit, which also destroyed FWindowsStylusInputInterface. Slate could then call CloseAllWindowsImmediately during Exit, which would try and access the deleted FWindowsStylusInputInterface via a delegate that was bound to RemovePluginForWindow with AddRaw

Since FWindowsStylusInputInterface was already allocated as a TSharedPtr, it now also inherits from TSharedFromThis so that the RemovePluginForWindow binding can be made with AddSP rather than AddRaw

diff --git a/Engine/Plugins/Editor/StylusInput/Source/StylusInput/Private/WindowsStylusInputInterface.cpp b/Engine/Plugins/Editor/StylusInput/Source/StylusInput/Private/WindowsStylusInputInterface.cpp
--- a/Engine/Plugins/Editor/StylusInput/Source/StylusInput/Private/WindowsStylusInputInterface.cpp
+++ b/Engine/Plugins/Editor/StylusInput/Source/StylusInput/Private/WindowsStylusInputInterface.cpp
@@ -138,7 +138,7 @@ void FWindowsStylusInputInterface::Tick()
 			if (CurrentPlugin == nullptr)
 			{
 				CreateStylusPluginForHWND(Hwnd);
-				Window->GetOnWindowClosedEvent().AddRaw(this, &FWindowsStylusInputInterface::RemovePluginForWindow);
+				Window->GetOnWindowClosedEvent().AddSP(this, &FWindowsStylusInputInterface::RemovePluginForWindow);
 			}
 		}
 	}
diff --git a/Engine/Plugins/Editor/StylusInput/Source/StylusInput/Private/WindowsStylusInputInterface.h b/Engine/Plugins/Editor/StylusInput/Source/StylusInput/Private/WindowsStylusInputInterface.h
--- a/Engine/Plugins/Editor/StylusInput/Source/StylusInput/Private/WindowsStylusInputInterface.h
+++ b/Engine/Plugins/Editor/StylusInput/Source/StylusInput/Private/WindowsStylusInputInterface.h
@@ -7,7 +7,7 @@
 class FWindowsStylusInputInterfaceImpl;
 class SWindow;
 
-class FWindowsStylusInputInterface : public IStylusInputInterfaceInternal
+class FWindowsStylusInputInterface : public IStylusInputInterfaceInternal, public TSharedFromThis<FWindowsStylusInputInterface>
 {
 public:
 	FWindowsStylusInputInterface(TUniquePtr<FWindowsStylusInputInterfaceImpl> InImpl);
