StylusInput: Merge in new API and Windows platform implementation

diff --git a/Engine/Plugins/Editor/StylusInput/Source/StylusInput/Private/DeprecatedStylusInputSubsystem.cpp b/Engine/Plugins/Editor/StylusInput/Source/StylusInput/Private/DeprecatedStylusInputSubsystem.cpp
new file mode 100644
--- /dev/null
+++ b/Engine/Plugins/Editor/StylusInput/Source/StylusInput/Private/DeprecatedStylusInputSubsystem.cpp
@@ -0,0 +1,102 @@
+﻿// Copyright Epic Games, Inc. All Rights Reserved.
+
+#include <Framework/Docking/TabManager.h>
+#include <Misc/App.h>
+
+#include "IStylusInputModule.h"
+
+#define LOCTEXT_NAMESPACE "FStylusInputModule"
+
+static const FName StylusInputDebugTabName = FName("StylusInputDebug");
+
+// This is the function that all platform-specific implementations are required to implement.
+TSharedPtr<IStylusInputInterfaceInternal> CreateStylusInputInterface();
+
+#if !PLATFORM_WINDOWS
+TSharedPtr<IStylusInputInterfaceInternal> CreateStylusInputInterface() { return TSharedPtr<IStylusInputInterfaceInternal>(); }
+#endif
+
+void UStylusInputSubsystem::Initialize(FSubsystemCollectionBase& Collection)
+{
+	if (FApp::IsUnattended() || IsRunningCommandlet())
+	{
+		return;
+	}
+
+	Super::Initialize(Collection);
+
+	UE_LOG(LogStylusInput, Log, TEXT("Initializing StylusInput subsystem."));
+
+	InputInterface = CreateStylusInputInterface();
+
+	if (!InputInterface.IsValid())
+	{
+		UE_LOG(LogStylusInput, Log, TEXT("StylusInput not supported on this platform."));
+		return;
+	}
+}
+
+void UStylusInputSubsystem::Deinitialize()
+{
+	Super::Deinitialize();
+
+	FGlobalTabmanager::Get()->UnregisterTabSpawner(StylusInputDebugTabName);
+
+	InputInterface.Reset();
+
+	UE_LOG(LogStylusInput, Log, TEXT("Shutting down StylusInput subsystem."));
+}
+
+int32 UStylusInputSubsystem::NumInputDevices() const
+{
+	if (InputInterface.IsValid())
+	{
+		return InputInterface->NumInputDevices();
+	}
+	return 0;
+}
+
+const IStylusInputDevice* UStylusInputSubsystem::GetInputDevice(int32 Index) const
+{
+	if (InputInterface.IsValid())
+	{
+		return InputInterface->GetInputDevice(Index);
+	}
+	return nullptr;
+}
+
+void UStylusInputSubsystem::AddMessageHandler(IStylusMessageHandler& InHandler)
+{
+	MessageHandlers.AddUnique(&InHandler);
+}
+
+void UStylusInputSubsystem::RemoveMessageHandler(IStylusMessageHandler& InHandler)
+{
+	MessageHandlers.Remove(&InHandler);
+}
+
+void UStylusInputSubsystem::Tick(float DeltaTime)
+{
+	TRACE_CPUPROFILER_EVENT_SCOPE(UStylusInputSubsystem::Tick);
+
+	if (InputInterface.IsValid())
+	{
+		InputInterface->Tick();
+
+		for (int32 DeviceIdx = 0; DeviceIdx < NumInputDevices(); ++DeviceIdx)
+		{
+			IStylusInputDevice* InputDevice = InputInterface->GetInputDevice(DeviceIdx);
+			if (InputDevice->IsDirty())
+			{
+				InputDevice->Tick();
+
+				for (IStylusMessageHandler* Handler : MessageHandlers)
+				{
+					Handler->OnStylusStateChanged(InputDevice->GetCurrentState(), DeviceIdx);
+				}
+			}
+		}
+	}
+}
+
+#undef LOCTEXT_NAMESPACE
diff --git a/Engine/Plugins/Editor/StylusInput/Source/StylusInput/Private/DeprecatedWindowsStylusInputInterface.cpp b/Engine/Plugins/Editor/StylusInput/Source/StylusInput/Private/DeprecatedWindowsStylusInputInterface.cpp
new file mode 100644
--- /dev/null
+++ b/Engine/Plugins/Editor/StylusInput/Source/StylusInput/Private/DeprecatedWindowsStylusInputInterface.cpp
@@ -0,0 +1,208 @@
+﻿// Copyright Epic Games, Inc. All Rights Reserved.
+
+#include "DeprecatedWindowsStylusInputInterface.h"
+
+#if PLATFORM_WINDOWS
+
+#include <Algo/Find.h>
+#include <Framework/Application/SlateApplication.h>
+
+#include "StylusInputTabletContext.h"
+
+FDeprecatedWindowsStylusInputInterface::FDeprecatedWindowsStylusInputInterface()
+	: TabletContexts(MakeUnique<TArray<FDeprecatedStylusInputDevice>>())
+{
+}
+
+FDeprecatedWindowsStylusInputInterface::~FDeprecatedWindowsStylusInputInterface() = default;
+
+void FDeprecatedWindowsStylusInputInterface::Tick()
+{
+	for (const FDeprecatedStylusInputDevice& TabletContext : *TabletContexts)
+	{
+		// don't change focus if any stylus is down
+		if (TabletContext.GetCurrentState().IsStylusDown())
+		{
+			return;
+		}
+	}
+
+	if (!FSlateApplication::IsInitialized())
+	{
+		return;
+	}
+
+	FSlateApplication& Application = FSlateApplication::Get();
+
+	FWidgetPath WidgetPath = Application.LocateWindowUnderMouse(Application.GetCursorPos(), Application.GetInteractiveTopLevelWindows());
+	if (WidgetPath.IsValid())
+	{
+		const TSharedPtr<SWindow> Window = WidgetPath.GetWindow();
+		if (Window.IsValid() && Window->IsRegularWindow())
+		{
+			TUniquePtr<FStylusInputInstanceWrapper> *const StylusInputInstance = StylusInputInstances.Find(Window.Get());
+			if (StylusInputInstance == nullptr)
+			{
+				CreateStylusInputInstance(Window.Get());
+				Window->GetOnWindowClosedEvent().AddSP(this, &FDeprecatedWindowsStylusInputInterface::RemoveStylusInputInstance);
+			}
+		}
+	}
+}
+
+int32 FDeprecatedWindowsStylusInputInterface::NumInputDevices() const
+{
+	return TabletContexts->Num();
+}
+
+IStylusInputDevice* FDeprecatedWindowsStylusInputInterface::GetInputDevice(int32 Index) const
+{
+	return 0 <= Index && Index < TabletContexts->Num() ? &(*TabletContexts)[Index] : nullptr;
+}
+
+void FDeprecatedWindowsStylusInputInterface::CreateStylusInputInstance(SWindow* Window)
+{	
+	StylusInputInstances.Emplace(Window, MakeUnique<FStylusInputInstanceWrapper>(Window, *TabletContexts));
+}
+
+void FDeprecatedWindowsStylusInputInterface::RemoveStylusInputInstance(const TSharedRef<SWindow>& Window)
+{
+	StylusInputInstances.Remove(&Window.Get());
+}
+
+FStylusState ToDeprecatedStylusState(const UE::StylusInput::FStylusInputPacket& Packet)
+{
+	return {
+		{Packet.X, Packet.Y}, Packet.Z, {Packet.XTiltOrientation, Packet.YTiltOrientation}, Packet.TwistOrientation, Packet.NormalPressure,
+		Packet.TangentPressure, {Packet.Width, Packet.Height},
+		(Packet.PenStatus & UE::StylusInput::EPenStatus::CursorIsTouching) != UE::StylusInput::EPenStatus::None,
+		(Packet.PenStatus & UE::StylusInput::EPenStatus::CursorIsInverted) != UE::StylusInput::EPenStatus::None
+	};
+}
+
+FDeprecatedWindowsStylusInputInterface::FDeprecatedStylusInputDevice::FDeprecatedStylusInputDevice(uint32 TabletContextId, const TSharedPtr<UE::StylusInput::IStylusInputTabletContext>& TabletContext)
+	: TabletContextId(TabletContextId)
+{
+	if (TabletContext)
+	{
+		UE::StylusInput::ETabletSupportedProperties SupportedProperties = TabletContext->GetSupportedProperties();
+
+		auto IsSupported = [SupportedProperties](UE::StylusInput::ETabletSupportedProperties Property)
+		{
+			return (SupportedProperties & Property) != UE::StylusInput::ETabletSupportedProperties::None;
+		};
+
+		if (IsSupported(UE::StylusInput::ETabletSupportedProperties::X) && IsSupported(UE::StylusInput::ETabletSupportedProperties::Y))
+		{
+			SupportedInputs.Add(EStylusInputType::Position);
+		}
+
+		if (IsSupported(UE::StylusInput::ETabletSupportedProperties::Z))
+		{
+			SupportedInputs.Add(EStylusInputType::Z);
+		}
+
+		if (IsSupported(UE::StylusInput::ETabletSupportedProperties::NormalPressure))
+		{
+			SupportedInputs.Add(EStylusInputType::Pressure);
+		}
+
+		if (IsSupported(UE::StylusInput::ETabletSupportedProperties::XTiltOrientation) && IsSupported(UE::StylusInput::ETabletSupportedProperties::YTiltOrientation))
+		{
+			SupportedInputs.Add(EStylusInputType::Tilt);
+		}
+
+		if (IsSupported(UE::StylusInput::ETabletSupportedProperties::TangentPressure))
+		{
+			SupportedInputs.Add(EStylusInputType::TangentPressure);
+		}
+
+		if (IsSupported(UE::StylusInput::ETabletSupportedProperties::ButtonPressure))
+		{
+			SupportedInputs.Add(EStylusInputType::ButtonPressure);
+		}
+
+		if (IsSupported(UE::StylusInput::ETabletSupportedProperties::TwistOrientation))
+		{
+			SupportedInputs.Add(EStylusInputType::Tilt);
+		}
+
+		if (IsSupported(UE::StylusInput::ETabletSupportedProperties::Width) && IsSupported(UE::StylusInput::ETabletSupportedProperties::Height))
+		{
+			SupportedInputs.Add(EStylusInputType::Size);
+		}
+	}
+}
+
+void FDeprecatedWindowsStylusInputInterface::FDeprecatedStylusInputDevice::Tick()
+{
+	PreviousState = CurrentState;
+	CurrentState = ToDeprecatedStylusState(LastPacket);
+	Dirty = false;
+}
+
+FDeprecatedWindowsStylusInputInterface::FStylusInputEventHandler::FStylusInputEventHandler(UE::StylusInput::IStylusInputInstance* Instance, TArray<FDeprecatedStylusInputDevice>& TabletContexts)
+	: Instance(Instance)
+	, TabletContexts(TabletContexts)
+{
+}
+
+void FDeprecatedWindowsStylusInputInterface::FStylusInputEventHandler::OnPacket(const UE::StylusInput::FStylusInputPacket& Packet)
+{
+	FDeprecatedStylusInputDevice* TabletContext = GetTabletContext(Packet.TabletContextID);
+
+	TabletContext->LastPacket = Packet;
+	TabletContext->SetDirty();
+}
+
+void FDeprecatedWindowsStylusInputInterface::FStylusInputEventHandler::OnDebugEvent(const FString& Message)
+{
+	// UE_LOG(LogStylusInput, Warning, TEXT("%s"), *Message);
+}
+
+FDeprecatedWindowsStylusInputInterface::FDeprecatedStylusInputDevice* FDeprecatedWindowsStylusInputInterface::FStylusInputEventHandler::GetTabletContext(
+	uint32 TabletContextId)
+{
+	FDeprecatedStylusInputDevice* TabletContext = Algo::FindByPredicate(
+		TabletContexts, [TabletContextId](const FDeprecatedStylusInputDevice& Value)
+		{
+			return
+				Value.TabletContextId == TabletContextId;
+		});
+
+	if (TabletContext)
+	{
+		return TabletContext;
+	}
+
+	TabletContext = &TabletContexts.Add_GetRef(FDeprecatedStylusInputDevice(TabletContextId, Instance->GetTabletContext(TabletContextId)));
+	
+	return TabletContext;
+}
+
+FDeprecatedWindowsStylusInputInterface::FStylusInputInstanceWrapper::FStylusInputInstanceWrapper(SWindow* Window,
+                                                                                                 TArray<FDeprecatedStylusInputDevice>& TabletContexts)
+	: Instance(UE::StylusInput::CreateInstance(*Window))
+	, EventHandler(Instance, TabletContexts)
+{
+	if (Instance)
+	{
+		Instance->AddEventHandler(&EventHandler, UE::StylusInput::EEventHandlerThread::OnGameThread);
+	}
+}
+
+FDeprecatedWindowsStylusInputInterface::FStylusInputInstanceWrapper::~FStylusInputInstanceWrapper()
+{
+	if (Instance)
+	{
+		Instance->RemoveEventHandler(&EventHandler);
+		UE::StylusInput::ReleaseInstance(Instance);
+	}
+}
+
+TSharedPtr<IStylusInputInterfaceInternal> CreateStylusInputInterface()
+{
+	return MakeShared<FDeprecatedWindowsStylusInputInterface>();
+}
+
+#endif
diff --git a/Engine/Plugins/Editor/StylusInput/Source/StylusInput/Private/DeprecatedWindowsStylusInputInterface.h b/Engine/Plugins/Editor/StylusInput/Source/StylusInput/Private/DeprecatedWindowsStylusInputInterface.h
new file mode 100644
--- /dev/null
+++ b/Engine/Plugins/Editor/StylusInput/Source/StylusInput/Private/DeprecatedWindowsStylusInputInterface.h
@@ -0,0 +1,72 @@
+﻿// Copyright Epic Games, Inc. All Rights Reserved.
+
+#pragma once
+
+#if PLATFORM_WINDOWS
+
+#include "IStylusInputModule.h"
+#include "StylusInput.h"
+#include "StylusInputPacket.h"
+
+class SWindow;
+
+/**
+ * An implementation of the deprecated @see UStylusInputSubsystem interface by means of using the new interface and Windows implementation.
+ */
+class FDeprecatedWindowsStylusInputInterface : public IStylusInputInterfaceInternal, public TSharedFromThis<FDeprecatedWindowsStylusInputInterface>
+{
+public:
+	FDeprecatedWindowsStylusInputInterface();
+	virtual ~FDeprecatedWindowsStylusInputInterface();
+
+	virtual void Tick() override;
+	virtual int32 NumInputDevices() const override;
+	virtual IStylusInputDevice* GetInputDevice(int32 Index) const override;
+
+private:
+	void CreateStylusInputInstance(SWindow* Window);
+	void RemoveStylusInputInstance(const TSharedRef<SWindow>& Window);
+
+	struct FDeprecatedStylusInputDevice : IStylusInputDevice
+	{
+		explicit FDeprecatedStylusInputDevice(uint32 TabletContextId, const TSharedPtr<UE::StylusInput::IStylusInputTabletContext>& TabletContext);
+
+		void SetDirty() { Dirty = true; }
+		virtual void Tick() override;
+
+		uint32 TabletContextId;
+
+		UE::StylusInput::FStylusInputPacket LastPacket;
+	};
+
+	class FStylusInputEventHandler : public UE::StylusInput::IStylusInputEventHandler
+	{
+	public:
+		explicit FStylusInputEventHandler(UE::StylusInput::IStylusInputInstance* Instance, TArray<FDeprecatedStylusInputDevice>& TabletContexts);
+		virtual ~FStylusInputEventHandler() override = default;
+		virtual FString GetName() override { return "DeprecatedWindowsStylusInputInterfaceEventHandler"; }
+		virtual void OnPacket(const UE::StylusInput::FStylusInputPacket& Packet) override;
+		virtual void OnDebugEvent(const FString& Message) override;
+
+	private:
+		FDeprecatedStylusInputDevice* GetTabletContext(uint32 TabletContextId);
+
+		UE::StylusInput::IStylusInputInstance* const Instance;
+		TArray<FDeprecatedStylusInputDevice>& TabletContexts;
+	};
+
+	class FStylusInputInstanceWrapper
+	{
+	public:
+		explicit FStylusInputInstanceWrapper(SWindow* Window, TArray<FDeprecatedStylusInputDevice>& TabletContexts);
+		~FStylusInputInstanceWrapper();
+
+		UE::StylusInput::IStylusInputInstance* Instance;
+		FStylusInputEventHandler EventHandler;
+	};
+
+	TMap<SWindow*, TUniquePtr<FStylusInputInstanceWrapper>> StylusInputInstances;
+	TUniquePtr<TArray<FDeprecatedStylusInputDevice>> TabletContexts;
+};
+
+#endif
diff --git a/Engine/Plugins/Editor/StylusInput/Source/StylusInput/Private/SStylusInputDebugWidget.cpp b/Engine/Plugins/Editor/StylusInput/Source/StylusInput/Private/SStylusInputDebugWidget.cpp
deleted file mode 100644
diff --git a/Engine/Plugins/Editor/StylusInput/Source/StylusInput/Private/SStylusInputDebugWidget.h b/Engine/Plugins/Editor/StylusInput/Source/StylusInput/Private/SStylusInputDebugWidget.h
deleted file mode 100644
diff --git a/Engine/Plugins/Editor/StylusInput/Source/StylusInput/Private/StylusInput.cpp b/Engine/Plugins/Editor/StylusInput/Source/StylusInput/Private/StylusInput.cpp
new file mode 100644
--- /dev/null
+++ b/Engine/Plugins/Editor/StylusInput/Source/StylusInput/Private/StylusInput.cpp
@@ -0,0 +1,94 @@
+﻿// Copyright Epic Games, Inc. All Rights Reserved.
+
+#include "StylusInput.h"
+
+#include "StylusInputInterface.h"
+#include "StylusInputUtils.h"
+
+#if PLATFORM_WINDOWS
+#include <Windows/WindowsStylusInputInterface.h>
+#endif
+
+namespace UE::StylusInput
+{
+	namespace Private
+	{
+		class FStylusInputImpl;
+
+		class FNotSupportedStylusInputInterface final : public IStylusInputInterface
+		{
+		public:
+			virtual ~FNotSupportedStylusInputInterface() override = default;
+
+			virtual IStylusInputInstance* CreateInstance(SWindow& Window) override { return nullptr; }
+
+			virtual bool ReleaseInstance(IStylusInputInstance* Instance) override { return false; }
+
+		private:
+			friend FStylusInputImpl;
+			static TUniquePtr<IStylusInputInterface> Create()
+			{
+				UE_LOG(LogStylusInput, Warning, TEXT("Stylus input is not supported for this platform."));
+				return nullptr;
+			}
+		};
+
+		class FStylusInputImpl
+		{
+		public:
+			FStylusInputImpl()
+				: Interface(FPlatformInterface::Create())
+			{
+			}
+
+			IStylusInputInstance* CreateInstance(SWindow& Window) const
+			{
+				if (!Interface)
+				{
+					UE_LOG(LogStylusInput, Error, TEXT("Platform interface not available."));
+					return nullptr;
+				}
+
+				return Interface ? Interface->CreateInstance(Window) : nullptr;
+			}
+
+			bool ReleaseInstance(IStylusInputInstance* Instance) const
+			{
+				return Interface ? Interface->ReleaseInstance(Instance) : false;
+			}
+
+		private:
+			using FPlatformInterface =
+#if PLATFORM_WINDOWS
+				Windows::FWindowsStylusInputInterface
+#else
+				FNotSupportedStylusInputInterface
+#endif
+				;
+
+			TUniquePtr<IStylusInputInterface> Interface;
+		};
+	}
+
+	Private::FStylusInputImpl& GetImplSingleton()
+	{
+		static Private::FStylusInputImpl ImplSingleton;
+		return ImplSingleton;
+	}
+
+	IStylusInputInstance* CreateInstance(SWindow& Window)
+	{
+		return GetImplSingleton().CreateInstance(Window);
+	}
+
+	bool ReleaseInstance(IStylusInputInstance* Instance)
+	{
+		if (!Instance)
+		{
+			UE_LOG(Private::LogStylusInput, Warning, TEXT("Nullptr passed into FStylusInput::ReleaseInstance()."));
+			return false;
+		}
+
+		return GetImplSingleton().ReleaseInstance(Instance);
+	}
+}
diff --git a/Engine/Plugins/Editor/StylusInput/Source/StylusInput/Private/StylusInputInterface.h b/Engine/Plugins/Editor/StylusInput/Source/StylusInput/Private/StylusInputInterface.h
new file mode 100644
--- /dev/null
+++ b/Engine/Plugins/Editor/StylusInput/Source/StylusInput/Private/StylusInputInterface.h
@@ -0,0 +1,23 @@
+﻿// Copyright Epic Games, Inc. All Rights Reserved.
+
+#pragma once
+
+class SWindow;
+
+namespace UE::StylusInput
+{
+	class IStylusInputEventHandler;
+	class IStylusInputInstance;
+}
+
+namespace UE::StylusInput::Private
+{
+	class IStylusInputInterface
+	{
+	public:
+		virtual ~IStylusInputInterface() = default;
+
+		virtual IStylusInputInstance* CreateInstance(SWindow& Window) = 0;
+		virtual bool ReleaseInstance(IStylusInputInstance* Instance) = 0;
+	};
+}
diff --git a/Engine/Plugins/Editor/StylusInput/Source/StylusInput/Private/StylusInputModule.cpp b/Engine/Plugins/Editor/StylusInput/Source/StylusInput/Private/StylusInputModule.cpp
--- a/Engine/Plugins/Editor/StylusInput/Source/StylusInput/Private/StylusInputModule.cpp
+++ b/Engine/Plugins/Editor/StylusInput/Source/StylusInput/Private/StylusInputModule.cpp
@@ -1,137 +1,12 @@
 // Copyright Epic Games, Inc. All Rights Reserved.
 
-#include "IStylusInputModule.h"
+#include <Modules/ModuleManager.h>
 
-#include "Misc/App.h"
-#include "Widgets/Docking/SDockTab.h"
-#include "WorkspaceMenuStructure.h"
-#include "WorkspaceMenuStructureModule.h"
-#include "SStylusInputDebugWidget.h"
-
-#define LOCTEXT_NAMESPACE "FStylusInputModule"
-
-
-static const FName StylusInputDebugTabName = FName("StylusInputDebug");
-
-class FStylusInputModule : public IModuleInterface
+namespace UE::StylusInput::Private
 {
-};
-
-IMPLEMENT_MODULE(FStylusInputModule, StylusInput)
-
-// This is the function that all platform-specific implementations are required to implement.
-TSharedPtr<IStylusInputInterfaceInternal> CreateStylusInputInterface();
-
-#if PLATFORM_WINDOWS
-#include "WindowsStylusInputInterface.h"
-#else
-TSharedPtr<IStylusInputInterfaceInternal> CreateStylusInputInterface() { return TSharedPtr<IStylusInputInterfaceInternal>(); }
-#endif
-
-// TODO: Other platforms
-
-void UStylusInputSubsystem::Initialize(FSubsystemCollectionBase& Collection)
-{
-	if (FApp::IsUnattended() || IsRunningCommandlet())
+	class FStylusInputModule : public IModuleInterface
 	{
-		return;
-	}
-
-	Super::Initialize(Collection);
-
-	UE_LOG(LogStylusInput, Log, TEXT("Initializing StylusInput subsystem."));
-
-	InputInterface = CreateStylusInputInterface();
-
-	if (!InputInterface.IsValid())
-	{
-		UE_LOG(LogStylusInput, Log, TEXT("StylusInput not supported on this platform."));
-		return;
-	}
-
-	const TSharedRef<FGlobalTabmanager>& TabManager = FGlobalTabmanager::Get();
-	const IWorkspaceMenuStructure& MenuStructure = WorkspaceMenu::GetMenuStructure();
-
-	TabManager->RegisterNomadTabSpawner(StylusInputDebugTabName,
-		FOnSpawnTab::CreateUObject(this, &UStylusInputSubsystem::OnSpawnPluginTab))
-		.SetDisplayName(LOCTEXT("DebugTabTitle", "Stylus Input Debug"))
-		.SetTooltipText(LOCTEXT("DebugTabTooltip", "Debug panel to display current values of stylus inputs."))
-		.SetIcon(FSlateIcon(FAppStyle::GetAppStyleSetName(), "StylusInputDebug.TabIcon"))
-		.SetGroup(MenuStructure.GetDeveloperToolsDebugCategory());
-}
-
-void UStylusInputSubsystem::Deinitialize()
-{
-	Super::Deinitialize();
-
-	FGlobalTabmanager::Get()->UnregisterTabSpawner(StylusInputDebugTabName);
-
-	InputInterface.Reset();
-
-	UE_LOG(LogStylusInput, Log, TEXT("Shutting down StylusInput subsystem."));
+	};
 }
 
-int32 UStylusInputSubsystem::NumInputDevices() const
-{
-	if (InputInterface.IsValid())
-	{
-		return InputInterface->NumInputDevices();
-	}
-	return 0;
-}
-
-const IStylusInputDevice* UStylusInputSubsystem::GetInputDevice(int32 Index) const
-{
-	if (InputInterface.IsValid())
-	{
-		return InputInterface->GetInputDevice(Index);
-	}
-	return nullptr;
-}
-
-void UStylusInputSubsystem::AddMessageHandler(IStylusMessageHandler& InHandler)
-{
-	MessageHandlers.AddUnique(&InHandler);
-}
-
-void UStylusInputSubsystem::RemoveMessageHandler(IStylusMessageHandler& InHandler)
-{
-	MessageHandlers.Remove(&InHandler);
-}
-
-void UStylusInputSubsystem::Tick(float DeltaTime)
-{
-	TRACE_CPUPROFILER_EVENT_SCOPE(UStylusInputSubsystem::Tick);
-
-	if (InputInterface.IsValid())
-	{
-		InputInterface->Tick();
-
-		for (int32 DeviceIdx = 0; DeviceIdx < NumInputDevices(); ++DeviceIdx)
-		{
-			IStylusInputDevice* InputDevice = InputInterface->GetInputDevice(DeviceIdx);
-			if (InputDevice->IsDirty())
-			{
-				InputDevice->Tick();
-
-				for (IStylusMessageHandler* Handler : MessageHandlers)
-				{
-					Handler->OnStylusStateChanged(InputDevice->GetCurrentState(), DeviceIdx);
-				}
-			}
-		}
-	}
-}
-
-TSharedRef<SDockTab> UStylusInputSubsystem::OnSpawnPluginTab(const FSpawnTabArgs& Args)
-{
-	return SNew(SDockTab)
-		.TabRole(ETabRole::NomadTab)
-		[
-			SNew(SStylusInputDebugWidget, *this)
-		];
-}
-
-
-
-#undef LOCTEXT_NAMESPACE
+IMPLEMENT_MODULE(UE::StylusInput::Private::FStylusInputModule, StylusInput);
diff --git a/Engine/Plugins/Editor/StylusInput/Source/StylusInput/Private/StylusInputUtils.cpp b/Engine/Plugins/Editor/StylusInput/Source/StylusInput/Private/StylusInputUtils.cpp
new file mode 100644
--- /dev/null
+++ b/Engine/Plugins/Editor/StylusInput/Source/StylusInput/Private/StylusInputUtils.cpp
@@ -0,0 +1,8 @@
+﻿// Copyright Epic Games, Inc. All Rights Reserved.
+
+#include "StylusInputUtils.h"
+
+namespace UE::StylusInput::Private
+{
+	DEFINE_LOG_CATEGORY(LogStylusInput);
+}
diff --git a/Engine/Plugins/Editor/StylusInput/Source/StylusInput/Private/StylusInputUtils.h b/Engine/Plugins/Editor/StylusInput/Source/StylusInput/Private/StylusInputUtils.h
new file mode 100644
--- /dev/null
+++ b/Engine/Plugins/Editor/StylusInput/Source/StylusInput/Private/StylusInputUtils.h
@@ -0,0 +1,25 @@
+﻿// Copyright Epic Games, Inc. All Rights Reserved.
+
+#pragma once
+
+#include <Logging/LogMacros.h>
+
+namespace UE::StylusInput::Private
+{
+	DECLARE_LOG_CATEGORY_EXTERN(LogStylusInput, Log, All)
+
+	inline void LogError(const FString& Preamble, const FString& Message)
+	{
+		UE_LOG(LogStylusInput, Error, TEXT("%s: %s"), *Preamble, *Message);
+	}
+
+	inline void LogWarning(const FString& Preamble, const FString& Message)
+	{
+		UE_LOG(LogStylusInput, Warning,  TEXT("%s: %s"), *Preamble, *Message);
+	}
+
+	inline void LogVerbose(const FString& Preamble, const FString& Message)
+	{
+		UE_LOG(LogStylusInput, Verbose,  TEXT("%s: %s"), *Preamble, *Message);
+	}
+}
diff --git a/Engine/Plugins/Editor/StylusInput/Source/StylusInput/Private/Windows/WindowsStylusInputInstance.cpp b/Engine/Plugins/Editor/StylusInput/Source/StylusInput/Private/Windows/WindowsStylusInputInstance.cpp
new file mode 100644
--- /dev/null
+++ b/Engine/Plugins/Editor/StylusInput/Source/StylusInput/Private/Windows/WindowsStylusInputInstance.cpp
@@ -0,0 +1,454 @@
+﻿// Copyright Epic Games, Inc. All Rights Reserved.
+
+#if PLATFORM_WINDOWS
+
+#include "WindowsStylusInputInstance.h"
+
+#include <StylusInput.h>
+#include <StylusInputUtils.h>
+#include <Algo/Transform.h>
+
+#include "WindowsStylusInputPlatformAPI.h"
+#include "WindowsStylusInputPluginAsync.h"
+#include "WindowsStylusInputPluginSync.h"
+#include "WindowsStylusInputUtils.h"
+
+// Fixed indices at which the plugins are installed; we only ever have one each.
+#define ASYNC_PLUGIN_INDEX 0
+#define SYNC_PLUGIN_INDEX 0
+
+#define LOG_PREAMBLE "WindowsStylusInputInstance"
+
+namespace UE::StylusInput::Private::Windows
+{
+	FWindowsStylusInputInstance::FWindowsStylusInputInstance(const HWND OSWindowHandle)
+		: WindowsAPI(FWindowsStylusInputPlatformAPI::GetInstance())
+	{
+		Init(OSWindowHandle);
+	}
+
+	FWindowsStylusInputInstance::~FWindowsStylusInputInstance()
+	{
+		RealTimeStylus->RemoveAllStylusAsyncPlugins();
+		RealTimeStylus->RemoveAllStylusSyncPlugins();
+		RealTimeStylus->Release();
+	}
+
+	bool FWindowsStylusInputInstance::AddEventHandler(IStylusInputEventHandler* EventHandler, const EEventHandlerThread Thread)
+	{
+		if (!EventHandler)
+		{
+			LogWarning(LOG_PREAMBLE, "Tried to add nullptr as event handler.");
+			return false;
+		}
+
+		if (Thread == EEventHandlerThread::OnGameThread)
+		{
+			if (!AsyncPlugin)
+			{
+				EnablePlugin(EEventHandlerThread::OnGameThread);
+			}
+
+			if (!AsyncPlugin)
+			{
+				LogError(LOG_PREAMBLE, FString::Format(
+					         TEXT("Event handler '{0}' was not added since game thread stylus input plugin could not be installed."),
+					         {EventHandler->GetName()}));
+				return false;
+			}
+
+			return AsyncPlugin->AddEventHandler(EventHandler);
+		}
+		
+		if (Thread == EEventHandlerThread::Asynchronous)
+		{
+			if (!SyncPlugin)
+			{
+				EnablePlugin(EEventHandlerThread::Asynchronous);
+			}
+
+			if (!SyncPlugin)
+			{
+				LogError(LOG_PREAMBLE, FString::Format(
+							 TEXT("Event handler '{0}' was not added since asynchronous stylus input plugin could not be installed."),
+							 {EventHandler->GetName()}));
+				return false;
+			}
+
+			return SyncPlugin->AddEventHandler(EventHandler);
+		}
+
+		return false;
+	}
+
+	bool FWindowsStylusInputInstance::RemoveEventHandler(IStylusInputEventHandler* EventHandler)
+	{
+		if (!EventHandler)
+		{
+			LogWarning(LOG_PREAMBLE, "Tried to remove nullptr event handler.");
+			return false;
+		}
+
+		bool bWasRemoved = false;
+		
+		if (AsyncPlugin && AsyncPlugin->RemoveEventHandler(EventHandler))
+		{
+			bWasRemoved = true;
+
+			if (AsyncPlugin->NumEventHandlers() == 0)
+			{
+				DisablePlugin(EEventHandlerThread::OnGameThread);
+			}
+		}
+
+		if (SyncPlugin && SyncPlugin->RemoveEventHandler(EventHandler))
+		{
+			bWasRemoved = true;
+
+			if (SyncPlugin->NumEventHandlers() == 0)
+			{
+				DisablePlugin(EEventHandlerThread::Asynchronous);
+			}
+		}
+
+		if (!bWasRemoved)
+		{
+			LogError(LOG_PREAMBLE, FString::Format(TEXT("Event handler '%s' does not exist."), {EventHandler->GetName()}));
+		}
+
+		return bWasRemoved;
+	}
+
+	const TSharedPtr<IStylusInputTabletContext> FWindowsStylusInputInstance::GetTabletContext(const uint32 TabletContextID)
+	{
+		return TabletContexts.Get(TabletContextID);
+	}
+
+	const TSharedPtr<IStylusInputStylusInfo> FWindowsStylusInputInstance::GetStylusInfo(uint32 StylusID)
+	{
+		TSharedPtr<FStylusInfo> StylusInfo = StylusInfos.Get(StylusID);
+		if (StylusInfo.IsValid())
+		{
+			return StylusInfo;
+		}
+
+		StylusInfo = StylusInfos.Add(StylusID);
+
+		IInkCursor* InkCursor = nullptr;
+		if (Succeeded(RealTimeStylus->GetStylusForId(StylusID, &InkCursor), LOG_PREAMBLE) && InkCursor != nullptr)
+		{
+			BSTR String = nullptr;
+			if (Succeeded(InkCursor->get_Name(&String), LOG_PREAMBLE))
+			{
+				StylusInfo->Name = String;
+				WindowsAPI.SysFreeString(String);
+			}
+			else
+			{
+				LogError(LOG_PREAMBLE, FString::Format(TEXT("Could not get name for stylus with ID {0}."), {StylusID}));
+			}
+
+			IInkCursorButtons* Buttons = nullptr;
+			long NumButtons = 0;
+			if (Succeeded(InkCursor->get_Buttons(&Buttons), LOG_PREAMBLE) &&
+				Buttons && 
+				Succeeded(Buttons->get_Count(&NumButtons), LOG_PREAMBLE))
+			{
+				if (NumButtons > 0)
+				{
+					TArray<FStylusButton> StylusInfoButtons;
+					StylusInfoButtons.Reserve(NumButtons);
+
+					for (long ButtonIndex = 0; ButtonIndex < NumButtons; ++ButtonIndex)
+					{
+						VARIANT ButtonIdentifier;
+						WindowsAPI.VariantInit(&ButtonIdentifier);
+
+						ButtonIdentifier.vt = VT_INT;
+						ButtonIdentifier.lVal = ButtonIndex;
+
+						IInkCursorButton* Button = nullptr;
+						if (Succeeded(Buttons->Item(ButtonIdentifier, &Button), LOG_PREAMBLE))
+						{
+							FStylusButton NewButton;
+
+							if (Succeeded(Button->get_Id(&String), LOG_PREAMBLE))
+							{
+								NewButton.ID = String;
+								WindowsAPI.SysFreeString(String);
+							}
+							else
+							{
+								LogError(LOG_PREAMBLE, FString::Format(
+											 TEXT("Could not get ID for button {0} for stylus with ID {1}."), {static_cast<int32>(ButtonIndex), StylusID}));
+							}
+
+							if (Succeeded(Button->get_Name(&String), LOG_PREAMBLE))
+							{
+								NewButton.Name = String;
+								WindowsAPI.SysFreeString(String);
+							}
+							else
+							{
+								LogError(LOG_PREAMBLE, FString::Format(
+											 TEXT("Could not get name for button {0} for stylus with ID {1}."), {static_cast<int32>(ButtonIndex), StylusID}));
+							}
+
+							if (!NewButton.ID.IsEmpty() && !NewButton.Name.IsEmpty())
+							{
+								StylusInfoButtons.Emplace(NewButton);
+							}
+						}
+						else
+						{
+							LogError(LOG_PREAMBLE, FString::Format(
+										 TEXT("Could not get button {0} for stylus with ID {1}."), {static_cast<int32>(ButtonIndex), StylusID}));
+						}
+
+						WindowsAPI.VariantClear(&ButtonIdentifier);
+					}
+
+					StylusInfo->Buttons = MoveTemp(StylusInfoButtons);
+				}
+			}
+			else
+			{
+				LogError(LOG_PREAMBLE, FString::Format(TEXT("Could not get buttons for stylus with ID {0}."), {StylusID}));
+			}
+		}
+		else
+		{
+			LogWarning(LOG_PREAMBLE, FString::Format(TEXT("Could not get stylus info for ID {0}."), {StylusID}));
+		}
+
+		return StylusInfo;
+	}
+
+	float FWindowsStylusInputInstance::GetPacketsPerSecond(const EEventHandlerThread EventHandlerThread) const
+	{
+		if (EventHandlerThread == EEventHandlerThread::Asynchronous)
+		{
+			return SyncPlugin ? SyncPlugin->GetPacketsPerSecond() : 0.0f;
+		}
+
+		if (EventHandlerThread == EEventHandlerThread::OnGameThread)
+		{
+			return AsyncPlugin ? AsyncPlugin->GetPacketsPerSecond() : 0.0f;
+		}
+
+		return -1.0f;
+	}
+
+	void FWindowsStylusInputInstance::SetupWindowContext(const HWND HWindow)
+	{
+		RECT WindowRectangle;
+		if (!ensure(WindowsAPI.GetClientRect(HWindow, &WindowRectangle)))
+		{
+			LogError(LOG_PREAMBLE, "Could not retrieve window rectangle; failed to setup device context!");
+			return;
+		}
+
+		if (WindowRectangle.right <= WindowRectangle.left || WindowRectangle.bottom <= WindowRectangle.top)
+		{
+			LogError(LOG_PREAMBLE, "Window is minimized; failed to setup device context!");
+			return;
+		}
+
+		const HDC HWindowDeviceContext = WindowsAPI.GetDC(HWindow);
+		if (!ensure(HWindowDeviceContext))
+		{
+			LogError(LOG_PREAMBLE, "Could not retrieve window device context; failed to setup device context!");
+			return;
+		}
+
+		const FVector2d DisplaySizeInMillimeters = {
+			static_cast<double>(WindowsAPI.GetDeviceCaps(HWindowDeviceContext, HORZSIZE)),
+			static_cast<double>(WindowsAPI.GetDeviceCaps(HWindowDeviceContext, VERTSIZE))
+		};
+
+		if (!ensure(DisplaySizeInMillimeters.X > 0.0 && DisplaySizeInMillimeters.Y > 0.0))
+		{
+			LogError(LOG_PREAMBLE, "Display size in millimeters is invalid; failed to setup device context!");
+			return;
+		}
+
+		const FVector2d DisplaySizeInPixels = {
+			static_cast<double>(WindowsAPI.GetDeviceCaps(HWindowDeviceContext, HORZRES)),
+			static_cast<double>(WindowsAPI.GetDeviceCaps(HWindowDeviceContext, VERTRES))
+		};
+		
+		if (!ensure(DisplaySizeInPixels.X > 0.0 && DisplaySizeInPixels.Y > 0.0))
+		{
+			LogError(LOG_PREAMBLE, "Display size in pixels is invalid; failed to setup device context!");
+			return;
+		}
+
+		const FVector2d DisplayPixelsPerLogicalInch = {
+			static_cast<double>(WindowsAPI.GetDeviceCaps(HWindowDeviceContext, LOGPIXELSX)),
+			static_cast<double>(WindowsAPI.GetDeviceCaps(HWindowDeviceContext, LOGPIXELSY))
+		};
+		
+		if (!ensure(DisplayPixelsPerLogicalInch.X > 0.0 && DisplayPixelsPerLogicalInch.Y > 0.0))
+		{
+			LogError(LOG_PREAMBLE, "Display pixels per logical inch is invalid; failed to setup device context!");
+			return;
+		}
+
+		const FVector2d WindowSizeInPixels = {
+			static_cast<double>(WindowRectangle.right - WindowRectangle.left),
+			static_cast<double>(WindowRectangle.bottom - WindowRectangle.top)
+		};
+
+		const FVector2d DisplayMillimetersPerPixel = DisplaySizeInMillimeters / DisplaySizeInPixels;
+		const FVector2d WindowSizeInMillimeters = WindowSizeInPixels * DisplayMillimetersPerPixel;
+		const FVector2d MillimetersPerInch{25.4};
+		const FVector2d DisplayInchesPerPixel = MillimetersPerInch / DisplayMillimetersPerPixel;
+
+		WindowContext.XYScale = DisplayInchesPerPixel / DisplayPixelsPerLogicalInch;
+		WindowContext.XYMaximum = WindowSizeInMillimeters * FVector2d{100.0} + FVector2d{0.5};
+		WindowContext.WindowSize = WindowSizeInPixels;
+
+		WindowsAPI.ReleaseDC(HWindow, HWindowDeviceContext);
+
+		LogVerbose(LOG_PREAMBLE, "Successfully setup window context.");
+	}
+
+	const FWindowContext& FWindowsStylusInputInstance::GetWindowContext() const
+	{
+		return WindowContext;
+	}
+
+	void FWindowsStylusInputInstance::UpdateTabletContexts(const FTabletContextContainer& InTabletContexts)
+	{
+		TabletContexts.Update(InTabletContexts);
+	}
+
+	void FWindowsStylusInputInstance::Init(HWND HWindow)
+	{
+		check(!RealTimeStylus.IsValid());
+		
+		SetupWindowContext(HWindow);
+
+		void* OutInstance = nullptr;
+		if (Failed(WindowsAPI.CoCreateInstance(CLSID_RealTimeStylus, nullptr, CLSCTX_ALL, IID_IRealTimeStylus, &OutInstance), LOG_PREAMBLE))
+		{			
+			LogError(LOG_PREAMBLE, "Could not create RealTimeStylus COM object instance!");
+		}
+		RealTimeStylus = static_cast<IRealTimeStylus*>(OutInstance);
+
+		if (RealTimeStylus.IsValid())
+		{
+			TArray<GUID> DesiredPackets;
+			DesiredPackets.Reserve(std::size(PacketPropertyConstants));
+			Algo::Transform(PacketPropertyConstants, DesiredPackets, [](const FPacketPropertyConstant& Property) { return Property.Guid; });
+
+			if (Failed(RealTimeStylus->SetDesiredPacketDescription(DesiredPackets.Num(), DesiredPackets.GetData()), LOG_PREAMBLE))
+			{
+				LogError(LOG_PREAMBLE, "Could not set desired packet description!");
+			}
+			
+			if (Failed(RealTimeStylus->put_HWND(reinterpret_cast<HANDLE_PTR>(HWindow)), LOG_PREAMBLE))
+			{
+				LogError(LOG_PREAMBLE, "Could not set window handle!");
+			}
+
+			if (Failed(RealTimeStylus->put_Enabled(::Windows::TRUE), LOG_PREAMBLE))
+			{
+				LogError(LOG_PREAMBLE, "Could not enable real time stylus input!");
+			}
+
+			LogVerbose(LOG_PREAMBLE, "Successfully initialized WindowsRealTimeStylus COM object.");
+		}
+	}
+
+	void FWindowsStylusInputInstance::EnablePlugin(const EEventHandlerThread EventHandlerThread)
+	{
+		check(RealTimeStylus.IsValid());
+
+		if (EventHandlerThread == EEventHandlerThread::OnGameThread)
+		{
+			AsyncPlugin = MakeUnique<FWindowsStylusInputPluginAsync>(
+				FGetWindowContextCallback::CreateRaw(this, &FWindowsStylusInputInstance::GetWindowContext),
+				FUpdateTabletContextsCallback::CreateRaw(this, &FWindowsStylusInputInstance::UpdateTabletContexts)
+			);
+
+			if (Succeeded(RealTimeStylus->AddStylusAsyncPlugin(ASYNC_PLUGIN_INDEX, AsyncPlugin.Get()), LOG_PREAMBLE))
+			{
+				LogVerbose(LOG_PREAMBLE, "Added game thread stylus input plugin!");
+			}
+			else
+			{
+				LogError(LOG_PREAMBLE, "Could not add game thread stylus input plugin!");
+				AsyncPlugin.Reset();
+			}
+		}
+
+		if (EventHandlerThread == EEventHandlerThread::Asynchronous)
+		{
+			SyncPlugin = MakeUnique<FWindowsStylusInputPluginSync>(
+				FGetWindowContextCallback::CreateRaw(this, &FWindowsStylusInputInstance::GetWindowContext),
+				FUpdateTabletContextsCallback::CreateRaw(this, &FWindowsStylusInputInstance::UpdateTabletContexts)
+			);
+
+			if (Succeeded(SyncPlugin->CreateFreeThreadMarshaler(), LOG_PREAMBLE))
+			{
+				if (Succeeded(RealTimeStylus->AddStylusSyncPlugin(SYNC_PLUGIN_INDEX, SyncPlugin.Get()), LOG_PREAMBLE))
+				{
+					LogVerbose(LOG_PREAMBLE, "Added asynchronous stylus input plugin!");
+				}
+				else
+				{
+					LogError(LOG_PREAMBLE, "Could not add asynchronous stylus input plugin!");
+					SyncPlugin.Reset();
+				}
+			}
+			else
+			{
+				LogError(LOG_PREAMBLE, "Could not create free thread marshaler for asynchronous stylus input plugin!");
+				SyncPlugin.Reset();
+			}
+		}
+	}
+
+	void FWindowsStylusInputInstance::DisablePlugin(const EEventHandlerThread EventHandlerThread)
+	{
+		check(RealTimeStylus.IsValid());
+
+#pragma warning(push)
+#pragma warning(disable : 6387)
+
+		if (EventHandlerThread == EEventHandlerThread::OnGameThread)
+		{
+			if (Failed(RealTimeStylus->RemoveStylusAsyncPlugin(ASYNC_PLUGIN_INDEX, nullptr), LOG_PREAMBLE))
+			{
+				LogError(LOG_PREAMBLE, "Could not remove game thread stylus input plugin!");
+			}
+			else
+			{
+				LogVerbose(LOG_PREAMBLE, "Removed game thread stylus input plugin!");
+			}
+
+			AsyncPlugin.Reset();
+		}
+
+		if (EventHandlerThread == EEventHandlerThread::Asynchronous)
+		{
+			if (Failed(RealTimeStylus->RemoveStylusSyncPlugin(SYNC_PLUGIN_INDEX, nullptr), LOG_PREAMBLE))
+			{
+				LogError(LOG_PREAMBLE, "Could not remove asynchronous stylus input plugin!");
+			}
+			else
+			{
+				LogVerbose(LOG_PREAMBLE, "Removed asynchronous stylus input plugin!");
+			}
+
+			SyncPlugin.Reset();
+		}
+
+#pragma warning(pop)
+	}
+}
+
+#undef LOG_PREAMBLE
+
+#endif
diff --git a/Engine/Plugins/Editor/StylusInput/Source/StylusInput/Private/Windows/WindowsStylusInputInstance.h b/Engine/Plugins/Editor/StylusInput/Source/StylusInput/Private/Windows/WindowsStylusInputInstance.h
new file mode 100644
--- /dev/null
+++ b/Engine/Plugins/Editor/StylusInput/Source/StylusInput/Private/Windows/WindowsStylusInputInstance.h
@@ -0,0 +1,68 @@
+﻿// Copyright Epic Games, Inc. All Rights Reserved.
+
+#pragma once
+
+#if PLATFORM_WINDOWS
+
+#include <StylusInput.h>
+#include <Containers/Array.h>
+#include <Templates/UniquePtr.h>
+
+#include <Windows/AllowWindowsPlatformTypes.h>
+	#include <Microsoft/COMPointer.h>
+	#include <RTSCom.h>
+#include <Windows/HideWindowsPlatformTypes.h>
+
+#include "WindowsStylusInputTabletContext.h"
+
+namespace UE::StylusInput
+{
+	class IStylusInputEventHandler;
+}
+
+namespace UE::StylusInput::Private::Windows
+{
+	class FWindowsStylusInputPluginSync;
+	class FWindowsStylusInputPluginAsync;
+	class FWindowsStylusInputPlatformAPI;
+
+	class FWindowsStylusInputInstance : public IStylusInputInstance
+	{
+	public:
+		explicit FWindowsStylusInputInstance(HWND OSWindowHandle);
+		virtual ~FWindowsStylusInputInstance() override;
+
+		virtual bool AddEventHandler(IStylusInputEventHandler* EventHandler, EEventHandlerThread Thread) override;
+		virtual bool RemoveEventHandler(IStylusInputEventHandler* EventHandler) override;
+
+		virtual const TSharedPtr<IStylusInputTabletContext> GetTabletContext(uint32 TabletContextID) override;
+		virtual const TSharedPtr<IStylusInputStylusInfo> GetStylusInfo(uint32 StylusID) override;
+
+		virtual float GetPacketsPerSecond(EEventHandlerThread EventHandlerThread) const override;
+
+	private:
+		void Init(HWND HWindow);
+
+		void EnablePlugin(EEventHandlerThread EventHandlerThread);
+		void DisablePlugin(EEventHandlerThread EventHandlerThread);
+
+		void SetupWindowContext(HWND HWindow);
+		const FWindowContext& GetWindowContext() const;
+
+		void UpdateTabletContexts(const FTabletContextContainer& InTabletContexts);
+
+		TComPtr<IRealTimeStylus> RealTimeStylus;
+
+		FWindowContext WindowContext;
+
+		TUniquePtr<FWindowsStylusInputPluginAsync> AsyncPlugin;
+		TUniquePtr<FWindowsStylusInputPluginSync> SyncPlugin;
+
+		FTabletContextThreadSafeContainer TabletContexts;
+		FStylusInfoThreadSafeContainer StylusInfos;
+
+		const FWindowsStylusInputPlatformAPI& WindowsAPI;
+	};
+}
+
+#endif
diff --git a/Engine/Plugins/Editor/StylusInput/Source/StylusInput/Private/Windows/WindowsStylusInputInterface.cpp b/Engine/Plugins/Editor/StylusInput/Source/StylusInput/Private/Windows/WindowsStylusInputInterface.cpp
new file mode 100644
--- /dev/null
+++ b/Engine/Plugins/Editor/StylusInput/Source/StylusInput/Private/Windows/WindowsStylusInputInterface.cpp
@@ -0,0 +1,94 @@
+﻿// Copyright Epic Games, Inc. All Rights Reserved.
+
+#if PLATFORM_WINDOWS
+
+#include "WindowsStylusInputInterface.h"
+
+#include <StylusInput.h>
+#include <StylusInputUtils.h>
+
+#include <GenericPlatform/GenericWindow.h>
+#include <Templates/SharedPointer.h>
+#include <Widgets/SWindow.h>
+
+#include <Windows/AllowWindowsPlatformTypes.h>
+	#include <RTSCOM_i.c>
+#include <Windows/HideWindowsPlatformTypes.h>
+
+#include "WindowsStylusInputInstance.h"
+#include "WindowsStylusInputPlatformAPI.h"
+
+namespace UE::StylusInput::Private::Windows
+{
+	IStylusInputInstance* FWindowsStylusInputInterface::CreateInstance(SWindow& Window)
+	{
+		if (FRefCountedInstance* ExistingRefCountedInstance = Instances.Find(&Window))
+		{
+			++ExistingRefCountedInstance->RefCount;
+			return ExistingRefCountedInstance->Instance.Get();
+		}
+
+		HWND OSWindowHandle = [&Window]
+		{
+			const TSharedPtr<const FGenericWindow> NativeWindow = Window.GetNativeWindow();
+			return NativeWindow.IsValid() ? static_cast<HWND>(NativeWindow->GetOSWindowHandle()) : nullptr;
+		}();
+
+		if (!OSWindowHandle)
+		{
+			LogError("WindowsStylusInputInterface", "Could not get native window handle.");
+			return nullptr;
+		}
+
+		FWindowsStylusInputInstance* NewInstance = Instances.Emplace(&Window, {MakeUnique<FWindowsStylusInputInstance>(OSWindowHandle), 1}).Instance.Get();
+		if (!ensureMsgf(NewInstance, TEXT("WindowsStylusInputInterface: Failed to create stylus input instance.")))
+		{
+			Instances.Remove(&Window);
+			return nullptr;
+		}
+
+		return NewInstance;
+	}
+
+	bool FWindowsStylusInputInterface::ReleaseInstance(IStylusInputInstance* Instance)
+	{
+		check(Instance);
+
+		FWindowsStylusInputInstance *const WindowsInstance = static_cast<FWindowsStylusInputInstance*>(Instance);
+
+		// Find existing instance
+		for (TTuple<SWindow*, FRefCountedInstance>& Entry : Instances)
+		{
+			FRefCountedInstance& RefCountedInstance = Entry.Get<1>();
+
+			if (RefCountedInstance.Instance.Get() == WindowsInstance)
+			{
+				// Decrease reference count
+				check(RefCountedInstance.RefCount > 0);
+				if (--RefCountedInstance.RefCount == 0)
+				{
+					// Delete if there are no references left
+					Instances.Remove(Entry.Key);
+				}
+					
+				return true;
+			}
+		}
+
+		ensureMsgf(false, TEXT("WindowsStylusInputInterface: Failed to find provied instance."));
+		return false;
+	}
+
+	TUniquePtr<IStylusInputInterface> FWindowsStylusInputInterface::Create()
+	{
+		if (FWindowsStylusInputPlatformAPI::GetInstance().SatisfiesRequirements())
+		{
+			TUniquePtr<FWindowsStylusInputInterface> WindowsImpl = MakeUnique<FWindowsStylusInputInterface>();
+			return TUniquePtr<IStylusInputInterface>(MoveTemp(WindowsImpl));
+		}
+
+		return nullptr;
+	}
+}
+
+#endif
diff --git a/Engine/Plugins/Editor/StylusInput/Source/StylusInput/Private/Windows/WindowsStylusInputInterface.h b/Engine/Plugins/Editor/StylusInput/Source/StylusInput/Private/Windows/WindowsStylusInputInterface.h
new file mode 100644
--- /dev/null
+++ b/Engine/Plugins/Editor/StylusInput/Source/StylusInput/Private/Windows/WindowsStylusInputInterface.h
@@ -0,0 +1,42 @@
+﻿// Copyright Epic Games, Inc. All Rights Reserved.
+
+#pragma once
+
+#if PLATFORM_WINDOWS
+
+#include <StylusInputInterface.h>
+#include <Containers/Array.h>
+#include <Containers/Map.h>
+#include <Templates/UniquePtr.h>
+
+namespace UE::StylusInput::Private
+{
+	class FStylusInputImpl;
+}
+
+namespace UE::StylusInput::Private::Windows
+{
+	class FWindowsStylusInputInstance;
+
+	class FWindowsStylusInputInterface : public IStylusInputInterface
+	{
+	public:
+		virtual IStylusInputInstance* CreateInstance(SWindow& Window) override;
+		virtual bool ReleaseInstance(IStylusInputInstance* Instance) override;
+
+	private:
+		friend FStylusInputImpl;
+
+		static TUniquePtr<IStylusInputInterface> Create();
+
+		struct FRefCountedInstance
+		{
+			TUniquePtr<FWindowsStylusInputInstance> Instance;
+			int32 RefCount;
+		};
+
+		TMap<SWindow*, FRefCountedInstance> Instances;
+	};
+}
+
+#endif
diff --git a/Engine/Plugins/Editor/StylusInput/Source/StylusInput/Private/Windows/WindowsStylusInputPlatformAPI.cpp b/Engine/Plugins/Editor/StylusInput/Source/StylusInput/Private/Windows/WindowsStylusInputPlatformAPI.cpp
new file mode 100644
--- /dev/null
+++ b/Engine/Plugins/Editor/StylusInput/Source/StylusInput/Private/Windows/WindowsStylusInputPlatformAPI.cpp
@@ -0,0 +1,155 @@
+﻿// Copyright Epic Games, Inc. All Rights Reserved.
+
+#if PLATFORM_WINDOWS
+
+#include "WindowsStylusInputPlatformAPI.h"
+
+#include <KnownFolders.h>
+#include <shtypes.h>
+#include <StylusInputUtils.h>
+#include <Windows/WindowsPlatformProcess.h>
+
+#include "WindowsStylusInputUtils.h"
+
+#define LOG_PREAMBLE "WindowsPlatformAPI"
+
+namespace UE::StylusInput::Private::Windows
+{
+	FWindowsStylusInputPlatformAPI::FWindowsStylusInputPlatformAPI()
+	{
+		bInitializedComLibrary = FWindowsPlatformMisc::CoInitialize();
+		if (!bInitializedComLibrary)
+		{
+			LogError(LOG_PREAMBLE, "Could not initialize COM library.");
+			return;
+		}
+
+		auto GetDllHandle = [&DllHandles = DllHandles](const TCHAR* DllName) -> void* {
+			void* DllHandle = FWindowsPlatformProcess::GetDllHandle(DllName);
+			if (DllHandle)
+			{
+				DllHandles.Add(DllHandle);
+				LogVerbose(LOG_PREAMBLE, FString::Format(TEXT("Retrieved DLL handle for {0}."), {DllName}));
+			}
+			else
+			{
+				LogError(LOG_PREAMBLE, FString::Format(TEXT("Could not get DLL handle for {0}."), {DllName}));
+			}
+			return DllHandle;
+		};
+
+		auto GetDllExport = []<typename FFuncType>(const TCHAR* DllName, void* DllHandle, const TCHAR* ExportName, FFuncType& ExportHandle)
+		{
+			void* FuncPtr = FWindowsPlatformProcess::GetDllExport(DllHandle, ExportName);
+			if (FuncPtr)
+			{
+				LogVerbose(LOG_PREAMBLE, FString::Format(TEXT("Retrieved DLL export '{0}' from {1}."), {ExportName, DllName}));
+			}
+			else
+			{
+				LogError(LOG_PREAMBLE, FString::Format(TEXT("Could not get DLL export '{0}' in {1}."), {ExportName, DllName}));
+			}
+			ExportHandle = static_cast<FFuncType>(FuncPtr);
+		};
+
+		const TCHAR* User32DllName = TEXT("User32.dll");
+		void* User32DllHandle = GetDllHandle(User32DllName);
+		GetDllExport(User32DllName, User32DllHandle, TEXT("GetClientRect"), GetClientRect);
+		GetDllExport(User32DllName, User32DllHandle, TEXT("GetDC"), GetDC);
+		GetDllExport(User32DllName, User32DllHandle, TEXT("ReleaseDC"), ReleaseDC);
+
+		const TCHAR* Gdi32DllName = TEXT("Gdi32.dll");
+		void* Gdi32DllHandle = GetDllHandle(Gdi32DllName);
+		GetDllExport(Gdi32DllName, Gdi32DllHandle, TEXT("GetDeviceCaps"), GetDeviceCaps);
+
+		const TCHAR* Ole32DllName = TEXT("Ole32.dll");
+		void* Ole32DllHandle = GetDllHandle(Ole32DllName);
+		GetDllExport(Ole32DllName, Ole32DllHandle, TEXT("CoCreateInstance"), CoCreateInstance);
+		GetDllExport(Ole32DllName, Ole32DllHandle, TEXT("CoCreateFreeThreadedMarshaler"), CoCreateFreeThreadedMarshaler);
+		GetDllExport(Ole32DllName, Ole32DllHandle, TEXT("CoTaskMemFree"), CoTaskMemFree);
+		GetDllExport(Ole32DllName, Ole32DllHandle, TEXT("StringFromGUID2"), StringFromGUID2);
+
+		const TCHAR* OleAut32DllName = TEXT("OleAut32.dll");
+		void* OleAut32DllHandle = GetDllHandle(OleAut32DllName);
+		GetDllExport(OleAut32DllName, OleAut32DllHandle, TEXT("SysFreeString"), SysFreeString);
+		GetDllExport(OleAut32DllName, OleAut32DllHandle, TEXT("VariantClear"), VariantClear);
+		GetDllExport(OleAut32DllName, OleAut32DllHandle, TEXT("VariantInit"), VariantInit);
+
+		// Try to get handles to RTSCom.dll and InkObj.dll for WindowsInk.
+		{
+			const TCHAR* Shell32DllName = TEXT("Shell32.dll");
+			void* Shell32DllHandle = GetDllHandle(Shell32DllName);
+
+			using FFuncSHGetKnownFolderPath = HRESULT(REFKNOWNFOLDERID, DWORD, HANDLE, PWSTR*);
+			FFuncSHGetKnownFolderPath* SHGetKnownFolderPath = nullptr;
+			GetDllExport(Shell32DllName, Shell32DllHandle, TEXT("SHGetKnownFolderPath"), SHGetKnownFolderPath);
+
+			if (SHGetKnownFolderPath && CoTaskMemFree)
+			{
+				PWSTR ProgramFilesCommonPath = nullptr;
+				if (Succeeded(SHGetKnownFolderPath(FOLDERID_ProgramFilesCommon, 0, nullptr, &ProgramFilesCommonPath), LOG_PREAMBLE))
+				{
+					const FString InkDLLDirectory = FString(ProgramFilesCommonPath) + TEXT("\\microsoft shared\\ink");
+					FPlatformProcess::PushDllDirectory(*InkDLLDirectory);
+					{
+						const TCHAR* RTSComDllName = TEXT("RTSCom.dll");
+						void* RTSComDllHandle = GetDllHandle(RTSComDllName);
+						bHasRTSComDllHandle = RTSComDllHandle != nullptr;
+
+						const TCHAR* InkObjDllName = TEXT("InkObj.dll");
+						void* InkObjDllHandle = GetDllHandle(InkObjDllName);
+						bHasInkObjDllHandle = InkObjDllHandle != nullptr;
+					}
+					FPlatformProcess::PopDllDirectory(*InkDLLDirectory);
+				}
+
+				CoTaskMemFree(ProgramFilesCommonPath);
+			}
+		}
+	}
+
+	FWindowsStylusInputPlatformAPI::~FWindowsStylusInputPlatformAPI()
+	{
+		for (void* DllHandle : DllHandles)
+		{
+			FWindowsPlatformProcess::FreeDllHandle(DllHandle);
+		}
+
+		if (bInitializedComLibrary)
+		{
+			FWindowsPlatformMisc::CoUninitialize();
+		}
+	}
+
+	const FWindowsStylusInputPlatformAPI& FWindowsStylusInputPlatformAPI::GetInstance()
+	{
+		static FWindowsStylusInputPlatformAPI PlatformAPI;
+		return PlatformAPI;
+	}
+
+	bool FWindowsStylusInputPlatformAPI::SatisfiesRequirements() const
+	{
+		bool bSuccess = true;
+
+		bSuccess &= bInitializedComLibrary;
+		bSuccess &= GetClientRect != nullptr;
+		bSuccess &= GetDC != nullptr;
+		bSuccess &= ReleaseDC != nullptr;
+		bSuccess &= GetDeviceCaps != nullptr;
+		bSuccess &= CoCreateInstance != nullptr;
+		bSuccess &= CoCreateFreeThreadedMarshaler != nullptr;
+		bSuccess &= CoTaskMemFree != nullptr;
+		bSuccess &= StringFromGUID2 != nullptr;
+		bSuccess &= SysFreeString != nullptr;
+		bSuccess &= VariantClear != nullptr;
+		bSuccess &= VariantInit != nullptr;
+		bSuccess &= bHasRTSComDllHandle;
+		bSuccess &= bHasInkObjDllHandle;
+
+		return bSuccess;
+	}
+}
+
+#undef LOG_PREAMBLE
+
+#endif
diff --git a/Engine/Plugins/Editor/StylusInput/Source/StylusInput/Private/Windows/WindowsStylusInputPlatformAPI.h b/Engine/Plugins/Editor/StylusInput/Source/StylusInput/Private/Windows/WindowsStylusInputPlatformAPI.h
new file mode 100644
--- /dev/null
+++ b/Engine/Plugins/Editor/StylusInput/Source/StylusInput/Private/Windows/WindowsStylusInputPlatformAPI.h
@@ -0,0 +1,71 @@
+﻿// Copyright Epic Games, Inc. All Rights Reserved.
+
+#pragma once
+
+#if PLATFORM_WINDOWS
+
+#include <Containers/Array.h>
+
+#include <Windows/AllowWindowsPlatformTypes.h>
+	#include <windef.h>
+	#include <wtypes.h>
+	#include <Unknwn.h>
+#include <Windows/HideWindowsPlatformTypes.h>
+
+namespace UE::StylusInput::Private::Windows
+{
+	class FWindowsStylusInputPlatformAPI
+	{
+	public:
+		static const FWindowsStylusInputPlatformAPI& GetInstance();
+
+		bool SatisfiesRequirements() const;
+
+		using FFuncGetClientRect = BOOL(HWND, LPRECT);
+		FFuncGetClientRect* GetClientRect = nullptr;
+
+		using FFuncGetDC = HDC(HWND);
+		FFuncGetDC* GetDC = nullptr;
+
+		using FFuncReleaseDC = int(HWND, HDC);
+		FFuncReleaseDC* ReleaseDC = nullptr;
+
+		using FFuncGetDeviceCaps = int(HDC, int);
+		FFuncGetDeviceCaps* GetDeviceCaps = nullptr;
+
+		using FFuncCoCreateInstance = HRESULT(REFCLSID, IUnknown*, DWORD, REFIID, LPVOID);
+		FFuncCoCreateInstance* CoCreateInstance = nullptr;
+
+		using FFuncCoCreateFreeThreadedMarshaler = HRESULT(LPUNKNOWN, LPUNKNOWN*);
+		FFuncCoCreateFreeThreadedMarshaler* CoCreateFreeThreadedMarshaler = nullptr;
+
+		using FFuncCoTaskMemFree = void(LPVOID);
+		FFuncCoTaskMemFree* CoTaskMemFree = nullptr;
+
+		using FFuncStringFromGUID2 = int(REFGUID, LPOLESTR, int);
+		FFuncStringFromGUID2* StringFromGUID2 = nullptr;
+
+		using FFuncSysFreeString = void(BSTR);
+		FFuncSysFreeString* SysFreeString = nullptr;
+
+		using FFuncVariantClear = HRESULT(VARIANTARG*);
+		FFuncVariantClear* VariantClear = nullptr;
+
+		using FFuncVariantInit = void(VARIANTARG*);
+		FFuncVariantInit* VariantInit = nullptr;
+
+	private:
+
+		FWindowsStylusInputPlatformAPI();
+		~FWindowsStylusInputPlatformAPI();
+
+		TArray<void*> DllHandles;
+
+		bool bInitializedComLibrary = false;
+		bool bHasRTSComDllHandle = false;
+		bool bHasInkObjDllHandle = false;
+		
+	};
+}
+
+#endif
diff --git a/Engine/Plugins/Editor/StylusInput/Source/StylusInput/Private/Windows/WindowsStylusInputPluginAsync.cpp b/Engine/Plugins/Editor/StylusInput/Source/StylusInput/Private/Windows/WindowsStylusInputPluginAsync.cpp
new file mode 100644
--- /dev/null
+++ b/Engine/Plugins/Editor/StylusInput/Source/StylusInput/Private/Windows/WindowsStylusInputPluginAsync.cpp
@@ -0,0 +1,112 @@
+﻿// Copyright Epic Games, Inc. All Rights Reserved.
+
+#if PLATFORM_WINDOWS
+
+#include "WindowsStylusInputPluginAsync.h"
+
+namespace UE::StylusInput::Private::Windows
+{
+	FWindowsStylusInputPluginAsync::FWindowsStylusInputPluginAsync(FGetWindowContextCallback&& GetWindowContext,
+	                                                     FUpdateTabletContextsCallback&& UpdateTabletContextsCallback)
+		: FWindowsStylusInputPluginBase(MoveTemp(GetWindowContext), MoveTemp(UpdateTabletContextsCallback))
+	{
+	}
+
+	FWindowsStylusInputPluginAsync::~FWindowsStylusInputPluginAsync()
+	{
+	}
+
+	HRESULT FWindowsStylusInputPluginAsync::RealTimeStylusEnabled(IRealTimeStylus* RealTimeStylus, const ULONG TabletContextIDsCount, const TABLET_CONTEXT_ID* TabletContextIDs)
+	{
+		return ProcessRealTimeStylusEnabled(RealTimeStylus, TabletContextIDsCount, TabletContextIDs);
+	}
+
+	HRESULT FWindowsStylusInputPluginAsync::RealTimeStylusDisabled(IRealTimeStylus* RealTimeStylus, const ULONG TabletContextIDsCount, const TABLET_CONTEXT_ID* TabletContextIDs)
+	{
+		return ProcessRealTimeStylusDisabled(RealTimeStylus, TabletContextIDsCount, TabletContextIDs);
+	}
+
+	HRESULT FWindowsStylusInputPluginAsync::StylusDown(IRealTimeStylus* RealTimeStylus, const StylusInfo* StylusInfo, const ULONG PropertyCount, LONG* PacketBuffer, LONG**)
+	{
+		return ProcessPackets(StylusInfo, 1, PropertyCount, EPacketType::StylusDown, reinterpret_cast<int32*>(PacketBuffer));
+	}
+
+	HRESULT FWindowsStylusInputPluginAsync::StylusUp(IRealTimeStylus* RealTimeStylus, const StylusInfo* StylusInfo, const ULONG PropertyCount, LONG* PacketBuffer, LONG**)
+	{
+		return ProcessPackets(StylusInfo, 1, PropertyCount, EPacketType::StylusUp, reinterpret_cast<int32*>(PacketBuffer));
+	}
+
+	HRESULT FWindowsStylusInputPluginAsync::TabletAdded(IRealTimeStylus* RealTimeStylus, IInkTablet* Tablet)
+	{
+		return ProcessTabletAdded(Tablet);
+	}
+
+	HRESULT FWindowsStylusInputPluginAsync::TabletRemoved(IRealTimeStylus* RealTimeStylus, LONG TabletIndex)
+	{
+		return ProcessTabletRemoved(TabletIndex);
+	}
+
+	HRESULT FWindowsStylusInputPluginAsync::InAirPackets(IRealTimeStylus* RealTimeStylus, const StylusInfo* StylusInfo, const ULONG PacketCount,
+	                                                const ULONG PacketBufferLength, LONG* PacketBuffer, ULONG*, LONG**)
+	{
+		return ProcessPackets(StylusInfo, PacketCount, PacketBufferLength, EPacketType::AboveDigitizer, reinterpret_cast<int32*>(PacketBuffer));
+	}
+
+	HRESULT FWindowsStylusInputPluginAsync::Packets(IRealTimeStylus* RealTimeStylus, const StylusInfo* StylusInfo, const ULONG PacketCount,
+													const ULONG PacketBufferLength, LONG* PacketBuffer, ULONG*, LONG**)
+	{
+		return ProcessPackets(StylusInfo, PacketCount, PacketBufferLength, EPacketType::OnDigitizer, reinterpret_cast<int32*>(PacketBuffer));
+	}
+
+	HRESULT FWindowsStylusInputPluginAsync::StylusInRange(IRealTimeStylus* RealTimeStylus, TABLET_CONTEXT_ID TabletContextID, STYLUS_ID StylusID)
+	{
+		return E_NOTIMPL;
+	}
+
+	HRESULT FWindowsStylusInputPluginAsync::StylusOutOfRange(IRealTimeStylus* RealTimeStylus, TABLET_CONTEXT_ID TabletContextID, STYLUS_ID StylusID)
+	{
+		return E_NOTIMPL;
+	}
+
+	HRESULT FWindowsStylusInputPluginAsync::StylusButtonDown(IRealTimeStylus* RealTimeStylus, STYLUS_ID StylusID, const GUID* pGuidStylusButton, POINT* pStylusPos)
+	{		
+		return E_NOTIMPL;
+	}
+
+	HRESULT FWindowsStylusInputPluginAsync::StylusButtonUp(IRealTimeStylus* RealTimeStylus, STYLUS_ID StylusID, const GUID* pGuidStylusButton, POINT* pStylusPos)
+	{
+		return E_NOTIMPL;
+	}
+
+	HRESULT FWindowsStylusInputPluginAsync::CustomStylusDataAdded(IRealTimeStylus* piRtsSrc, const GUID* pGuidId, ULONG cbData, const BYTE* pbData)
+	{
+		return E_NOTIMPL;
+	}
+
+	HRESULT FWindowsStylusInputPluginAsync::SystemEvent(IRealTimeStylus* piRtsSrc, TABLET_CONTEXT_ID tcid, STYLUS_ID sid, SYSTEM_EVENT event, SYSTEM_EVENT_DATA eventdata)
+	{
+		return E_NOTIMPL;
+	}
+
+	HRESULT FWindowsStylusInputPluginAsync::Error(IRealTimeStylus*, IStylusPlugin* Plugin, const RealTimeStylusDataInterest DataInterest, const HRESULT ErrorCode, LONG_PTR*)
+	{
+#if STYLUSINPUT_SHOW_NOTIMPL_ERRORS
+		const FString ErrorMessage = ProcessError(DataInterest, ErrorCode);
+		DebugEvent(ErrorMessage);
+#endif
+
+		return S_OK;
+	}
+
+	HRESULT FWindowsStylusInputPluginAsync::UpdateMapping(IRealTimeStylus* piRtsSrc)
+	{
+		return E_NOTIMPL;
+	}
+
+	HRESULT FWindowsStylusInputPluginAsync::DataInterest(RealTimeStylusDataInterest* pDataInterest)
+	{
+		*pDataInterest = RTSDI_AllData;
+		return S_OK;
+	}
+}
+#endif
diff --git a/Engine/Plugins/Editor/StylusInput/Source/StylusInput/Private/Windows/WindowsStylusInputPluginAsync.h b/Engine/Plugins/Editor/StylusInput/Source/StylusInput/Private/Windows/WindowsStylusInputPluginAsync.h
new file mode 100644
--- /dev/null
+++ b/Engine/Plugins/Editor/StylusInput/Source/StylusInput/Private/Windows/WindowsStylusInputPluginAsync.h
@@ -0,0 +1,72 @@
+﻿// Copyright Epic Games, Inc. All Rights Reserved.
+
+#pragma once
+
+#if PLATFORM_WINDOWS
+
+#include "WindowsStylusInputPluginBase.h"
+
+namespace UE::StylusInput::Private::Windows
+{
+	class FWindowsStylusInputPluginAsync final : public IStylusAsyncPlugin, public FWindowsStylusInputPluginBase
+	{
+	public:
+		explicit FWindowsStylusInputPluginAsync(FGetWindowContextCallback&& GetWindowContext, FUpdateTabletContextsCallback&& UpdateTabletContextsCallback);
+		~FWindowsStylusInputPluginAsync();
+
+		// IUnknown
+		virtual HRESULT QueryInterface(REFIID InterfaceID, void** InterfaceObject) override
+		{
+			if (InterfaceID == IID_IStylusAsyncPlugin || InterfaceID == IID_IUnknown)
+			{
+				*InterfaceObject = this;
+				AddRef();
+				return S_OK;
+			}
+
+			*InterfaceObject = nullptr;
+			return E_NOINTERFACE;
+		}
+
+		virtual ULONG AddRef() override
+		{
+			return ++RefCount;
+		}
+
+		virtual ULONG Release() override
+		{
+			const int32 NewRefCount = --RefCount;
+			if (NewRefCount == 0)
+				delete this;
+
+			return NewRefCount;
+		}
+
+		// IStylusPlugin
+		virtual HRESULT RealTimeStylusEnabled(IRealTimeStylus* RealTimeStylus, ULONG TabletContextIDsCount, const TABLET_CONTEXT_ID* TabletContextIDs) override;
+		virtual HRESULT RealTimeStylusDisabled(IRealTimeStylus* RealTimeStylus, ULONG TabletContextIDsCount, const TABLET_CONTEXT_ID* TabletContextIDs) override;
+		virtual HRESULT StylusDown(IRealTimeStylus* RealTimeStylus, const StylusInfo* StylusInfo, ULONG PropertyCount, LONG* PacketBuffer, LONG** ppInOutPkt) override;
+		virtual HRESULT StylusUp(IRealTimeStylus* RealTimeStylus, const StylusInfo* StylusInfo, ULONG PropertyCount, LONG* PacketBuffer, LONG** ppInOutPkt) override;
+		virtual HRESULT TabletAdded(IRealTimeStylus* RealTimeStylus, IInkTablet* Tablet) override;
+		virtual HRESULT TabletRemoved(IRealTimeStylus* RealTimeStylus, LONG TabletIndex) override;
+		virtual HRESULT InAirPackets(IRealTimeStylus* RealTimeStylus, const StylusInfo* StylusInfo, ULONG PacketCount, ULONG PacketBufferLength, LONG* PacketBuffer, ULONG* InOutPacketCount, LONG** InOutPacketBuffer) override;
+		virtual HRESULT Packets(IRealTimeStylus* RealTimeStylus, const StylusInfo* StylusInfo, ULONG PacketCount, ULONG PacketBufferLength, LONG* PacketBuffer, ULONG* InOutPacketCount, LONG** InOutPacketBuffer) override;
+		virtual HRESULT StylusInRange(IRealTimeStylus* piRtsSrc, TABLET_CONTEXT_ID tcid, STYLUS_ID sid) override;
+		virtual HRESULT StylusOutOfRange(IRealTimeStylus* piRtsSrc, TABLET_CONTEXT_ID tcid, STYLUS_ID sid) override;
+		virtual HRESULT StylusButtonDown(IRealTimeStylus* piRtsSrc, STYLUS_ID sid, const GUID* pGuidStylusButton, POINT* pStylusPos) override;
+		virtual HRESULT StylusButtonUp(IRealTimeStylus* piRtsSrc, STYLUS_ID sid, const GUID* pGuidStylusButton, POINT* pStylusPos) override;
+		virtual HRESULT CustomStylusDataAdded(IRealTimeStylus* piRtsSrc, const GUID* pGuidId, ULONG cbData, const BYTE* pbData) override;
+		virtual HRESULT SystemEvent(IRealTimeStylus* piRtsSrc, TABLET_CONTEXT_ID tcid, STYLUS_ID sid, SYSTEM_EVENT event, SYSTEM_EVENT_DATA eventdata) override;
+		virtual HRESULT Error(IRealTimeStylus* RealTimeStylus, IStylusPlugin* Plugin, RealTimeStylusDataInterest DataInterest, HRESULT ErrorCode, LONG_PTR* InternalKey) override;
+		virtual HRESULT UpdateMapping(IRealTimeStylus* piRtsSrc) override;
+		virtual HRESULT DataInterest(RealTimeStylusDataInterest* pDataInterest) override;
+
+	protected:
+		virtual FString GetName() const override { return "OnGameThread"; }
+
+	private:
+		int32 RefCount = 1;
+	};
+}
+
+#endif
diff --git a/Engine/Plugins/Editor/StylusInput/Source/StylusInput/Private/Windows/WindowsStylusInputPluginBase.cpp b/Engine/Plugins/Editor/StylusInput/Source/StylusInput/Private/Windows/WindowsStylusInputPluginBase.cpp
new file mode 100644
--- /dev/null
+++ b/Engine/Plugins/Editor/StylusInput/Source/StylusInput/Private/Windows/WindowsStylusInputPluginBase.cpp
@@ -0,0 +1,717 @@
+﻿// Copyright Epic Games, Inc. All Rights Reserved.
+
+#if PLATFORM_WINDOWS
+
+#include "WindowsStylusInputPluginBase.h"
+
+#include <StylusInput.h>
+#include <StylusInputUtils.h>
+
+#include "WindowsStylusInputPlatformAPI.h"
+#include "WindowsStylusInputUtils.h"
+
+#include <Windows/AllowWindowsPlatformAtomics.h>
+	#include <comdef.h>
+#include <Windows/HideWindowsPlatformAtomics.h>
+
+#define LOG_PREAMBLE "WindowsStylusInputPluginBase"
+
+namespace UE::StylusInput::Private::Windows
+{
+	bool SetupTabletContextMetadata(IInkTablet& InkTablet, FTabletContext& TabletContext)
+	{
+		const FWindowsStylusInputPlatformAPI& WindowsAPI = FWindowsStylusInputPlatformAPI::GetInstance();
+
+		bool bSuccess = true;
+
+		BSTR String = nullptr;
+		if (Succeeded(InkTablet.get_Name(&String), LOG_PREAMBLE))
+		{
+			TabletContext.Name = String;
+			WindowsAPI.SysFreeString(String);
+		}
+		else
+		{
+			LogError(LOG_PREAMBLE, FString::Format(TEXT("Could not get name for TabletContext with ID {0}."), {TabletContext.ID}));
+			bSuccess = false;
+		}
+
+		if (Succeeded(InkTablet.get_PlugAndPlayId(&String), LOG_PREAMBLE))
+		{
+			TabletContext.PlugAndPlayID = FString(String);
+			WindowsAPI.SysFreeString(String);
+		}
+		else
+		{
+			LogError(LOG_PREAMBLE, FString::Format(TEXT("Could not get plug and play ID for TabletContext with ID {0}."), {TabletContext.ID}));
+			bSuccess = false;
+		}
+
+		IInkRectangle* Rectangle = nullptr;
+		if (Succeeded(InkTablet.get_MaximumInputRectangle(&Rectangle), LOG_PREAMBLE))
+		{
+			long Top, Left, Bottom, Right;
+			if (Succeeded(Rectangle->GetRectangle(&Top, &Left, &Bottom, &Right), LOG_PREAMBLE))
+			{
+				TabletContext.InputRectangle = {Left, Top, Right, Bottom};
+			}
+		}
+		else
+		{
+			LogError(LOG_PREAMBLE, FString::Format(TEXT("Could not get input rectangle for TabletContext with ID {0}."), {TabletContext.ID}));
+			bSuccess = false;
+		}
+
+		TabletHardwareCapabilities HardwareCapabilities;
+		if (Succeeded(InkTablet.get_HardwareCapabilities(&HardwareCapabilities), LOG_PREAMBLE))
+		{
+			TabletContext.HardwareCapabilities =
+				static_cast<ETabletHardwareCapabilities>(HardwareCapabilities & THWC_Integrated) |
+				static_cast<ETabletHardwareCapabilities>(HardwareCapabilities & THWC_CursorMustTouch) |
+				static_cast<ETabletHardwareCapabilities>(HardwareCapabilities & THWC_HardProximity) |
+				static_cast<ETabletHardwareCapabilities>(HardwareCapabilities & THWC_CursorsHavePhysicalIds);
+		}
+		else
+		{
+			LogError(LOG_PREAMBLE, FString::Format(TEXT("Could not get hardware capabilities for TabletContext with ID {0}."), {TabletContext.ID}));
+			bSuccess = false;
+		}
+
+		return bSuccess;
+	}
+
+	// TODO compact number of functions needed to only have differences in behavior, and then use offsetof() to apply to different members
+	// TODO XY validate mapping 
+
+	void SetPropertyX(FStylusInputPacket& Packet, const int32 Value, const int8* Data)
+	{
+		const double& Scale = reinterpret_cast<const double*>(Data)[0];
+		const double& Maximum = reinterpret_cast<const double*>(Data)[1];
+		const double& WindowWidth = reinterpret_cast<const double*>(Data)[2];
+
+		Packet.X = Value / Scale / Maximum * WindowWidth;
+	}
+	static_assert(std::is_same_v<FuncSetProperty, decltype(SetPropertyX)>);
+
+	void SetPropertyY(FStylusInputPacket& Packet, const int32 Value, const int8* Data)
+	{
+		const double& Scale = reinterpret_cast<const double*>(Data)[0];
+		const double& Maximum = reinterpret_cast<const double*>(Data)[1];
+		const double& WindowHeight = reinterpret_cast<const double*>(Data)[2];
+		
+		Packet.Y = Value / Scale / Maximum * WindowHeight;
+	}
+	static_assert(std::is_same_v<FuncSetProperty, decltype(SetPropertyY)>);
+
+	void SetPropertyZ(FStylusInputPacket& Packet, const int32 Value, const int8* Data)
+	{
+		const int32& Minimum = reinterpret_cast<const int32*>(Data)[0];
+		const int32& Maximum = reinterpret_cast<const int32*>(Data)[1];
+
+		Packet.Z = static_cast<float>(Value - Minimum) / (Maximum - Minimum);
+	}
+	static_assert(std::is_same_v<FuncSetProperty, decltype(SetPropertyZ)>);
+	
+	void SetPropertyPacketStatus(FStylusInputPacket& Packet, const int32 Value, const int8*)
+	{
+		Packet.PenStatus = static_cast<EPenStatus>(Value);
+	}
+	static_assert(std::is_same_v<FuncSetProperty, decltype(SetPropertyPacketStatus)>);
+
+	void SetPropertyTimerTick(FStylusInputPacket& Packet, const int32 Value, const int8* Data)
+	{
+		Packet.TimerTick = Value;
+	}
+	static_assert(std::is_same_v<FuncSetProperty, decltype(SetPropertyTimerTick)>);
+
+	void SetPropertySerialNumber(FStylusInputPacket& Packet, const int32 Value, const int8* Data)
+	{
+		Packet.SerialNumber = Value;
+	}
+	static_assert(std::is_same_v<FuncSetProperty, decltype(SetPropertySerialNumber)>);
+
+	void SetPropertyNormalPressure(FStylusInputPacket& Packet, const int32 Value, const int8* Data)
+	{
+		const int32& Minimum = reinterpret_cast<const int32*>(Data)[0];
+		const int32& Maximum = reinterpret_cast<const int32*>(Data)[1];
+
+		Packet.NormalPressure = static_cast<float>(Value - Minimum) / (Maximum - Minimum);
+	}
+	static_assert(std::is_same_v<FuncSetProperty, decltype(SetPropertyNormalPressure)>);
+
+	void SetPropertyTangentPressure(FStylusInputPacket& Packet, const int32 Value, const int8* Data)
+	{
+		const int32& Minimum = reinterpret_cast<const int32*>(Data)[0];
+		const int32& Maximum = reinterpret_cast<const int32*>(Data)[1];
+
+		Packet.TangentPressure = static_cast<float>(Value - Minimum) / (Maximum - Minimum);
+	}
+	static_assert(std::is_same_v<FuncSetProperty, decltype(SetPropertyTangentPressure)>);
+
+	void SetPropertyButtonPressure(FStylusInputPacket& Packet, const int32 Value, const int8* Data)
+	{
+		const int32& Minimum = reinterpret_cast<const int32*>(Data)[0];
+		const int32& Maximum = reinterpret_cast<const int32*>(Data)[1];
+
+		Packet.ButtonPressure = static_cast<float>(Value - Minimum) / (Maximum - Minimum);
+	}
+	static_assert(std::is_same_v<FuncSetProperty, decltype(SetPropertyButtonPressure)>);
+
+	void SetPropertyXTiltOrientation(FStylusInputPacket& Packet, const int32 Value, const int8* Data)
+	{
+		const float& Resolution = reinterpret_cast<const float*>(Data)[0];
+
+		Packet.XTiltOrientation = Value / Resolution;
+	}
+	static_assert(std::is_same_v<FuncSetProperty, decltype(SetPropertyXTiltOrientation)>);
+
+	void SetPropertyYTiltOrientation(FStylusInputPacket& Packet, const int32 Value, const int8* Data)
+	{
+		const float& Resolution = reinterpret_cast<const float*>(Data)[0];
+
+		Packet.YTiltOrientation = Value / Resolution;
+	}
+	static_assert(std::is_same_v<FuncSetProperty, decltype(SetPropertyYTiltOrientation)>);
+
+	void SetPropertyAzimuthOrientation(FStylusInputPacket& Packet, const int32 Value, const int8* Data)
+	{
+		const float& Resolution = reinterpret_cast<const float*>(Data)[0];
+
+		Packet.AzimuthOrientation = Value / Resolution;
+	}
+	static_assert(std::is_same_v<FuncSetProperty, decltype(SetPropertyAzimuthOrientation)>);
+
+	void SetPropertyAltitudeOrientation(FStylusInputPacket& Packet, const int32 Value, const int8* Data)
+	{
+		const float& Resolution = reinterpret_cast<const float*>(Data)[0];
+
+		Packet.AltitudeOrientation = Value / Resolution;
+	}
+	static_assert(std::is_same_v<FuncSetProperty, decltype(SetPropertyAltitudeOrientation)>);
+
+	void SetPropertyTwistOrientation(FStylusInputPacket& Packet, const int32 Value, const int8* Data)
+	{
+		const float& Resolution = reinterpret_cast<const float*>(Data)[0];
+
+		Packet.TwistOrientation = Value / Resolution;
+	}
+	static_assert(std::is_same_v<FuncSetProperty, decltype(SetPropertyTwistOrientation)>);
+
+	void SetPropertyPitchRotation(FStylusInputPacket& Packet, const int32 Value, const int8* Data)
+	{
+		const float& Resolution = reinterpret_cast<const float*>(Data)[0];
+
+		Packet.PitchRotation = Value / Resolution;
+	}
+	static_assert(std::is_same_v<FuncSetProperty, decltype(SetPropertyPitchRotation)>);
+
+	void SetPropertyRollRotation(FStylusInputPacket& Packet, const int32 Value, const int8* Data)
+	{
+		const float& Resolution = reinterpret_cast<const float*>(Data)[0];
+
+		Packet.RollRotation = Value / Resolution;
+	}
+	static_assert(std::is_same_v<FuncSetProperty, decltype(SetPropertyRollRotation)>);
+
+	void SetPropertyYawRotation(FStylusInputPacket& Packet, const int32 Value, const int8* Data)
+	{
+		const float& Resolution = reinterpret_cast<const float*>(Data)[0];
+
+		Packet.YawRotation = Value / Resolution;
+	}
+	static_assert(std::is_same_v<FuncSetProperty, decltype(SetPropertyYawRotation)>);
+
+	void SetPropertyWidth(FStylusInputPacket& Packet, const int32 Value, const int8* Data)
+	{
+		const float& Resolution = reinterpret_cast<const float*>(Data)[0];
+
+		Packet.Width = Value / Resolution;
+	}
+	static_assert(std::is_same_v<FuncSetProperty, decltype(SetPropertyWidth)>);
+
+	void SetPropertyHeight(FStylusInputPacket& Packet, const int32 Value, const int8* Data)
+	{
+		const float& Resolution = reinterpret_cast<const float*>(Data)[0];
+
+		Packet.Height = Value / Resolution;
+	}
+	static_assert(std::is_same_v<FuncSetProperty, decltype(SetPropertyHeight)>);
+
+	void SetPropertyFingerContactConfidence(FStylusInputPacket& Packet, const int32 Value, const int8* Data)
+	{
+		const int32& Minimum = reinterpret_cast<const int32*>(Data)[0];
+		const int32& Maximum = reinterpret_cast<const int32*>(Data)[1];
+
+		Packet.FingerContactConfidence = static_cast<float>(Value - Minimum) / (Maximum - Minimum);
+	}
+	static_assert(std::is_same_v<FuncSetProperty, decltype(SetPropertyFingerContactConfidence)>);
+
+	void SetPropertyDeviceContactID(FStylusInputPacket& Packet, const int32 Value, const int8* Data)
+	{
+		Packet.DeviceContactID = Value;
+	}
+	static_assert(std::is_same_v<FuncSetProperty, decltype(SetPropertyDeviceContactID)>);
+
+	FuncSetProperty* GetSetPropertyFunc(const EPacketPropertyType Type)
+	{
+		switch (Type)
+		{
+		case EPacketPropertyType::X:
+			return &SetPropertyX;
+		case EPacketPropertyType::Y:
+			return &SetPropertyY;
+		case EPacketPropertyType::Z:
+			return &SetPropertyZ;
+		case EPacketPropertyType::PacketStatus:
+			return &SetPropertyPacketStatus;
+		case EPacketPropertyType::TimerTick:
+			return &SetPropertyTimerTick;
+		case EPacketPropertyType::SerialNumber:
+			return &SetPropertySerialNumber;
+		case EPacketPropertyType::NormalPressure:
+			return &SetPropertyNormalPressure;
+		case EPacketPropertyType::TangentPressure:
+			return &SetPropertyTangentPressure;
+		case EPacketPropertyType::ButtonPressure:
+			return &SetPropertyButtonPressure;
+		case EPacketPropertyType::XTiltOrientation:
+			return &SetPropertyXTiltOrientation;
+		case EPacketPropertyType::YTiltOrientation:
+			return &SetPropertyYTiltOrientation;
+		case EPacketPropertyType::AzimuthOrientation:
+			return &SetPropertyAzimuthOrientation;
+		case EPacketPropertyType::AltitudeOrientation:
+			return &SetPropertyAltitudeOrientation;
+		case EPacketPropertyType::TwistOrientation:
+			return &SetPropertyTwistOrientation;
+		case EPacketPropertyType::PitchRotation:
+			return &SetPropertyPitchRotation;
+		case EPacketPropertyType::RollRotation:
+			return &SetPropertyRollRotation;
+		case EPacketPropertyType::YawRotation:
+			return &SetPropertyYawRotation;
+		case EPacketPropertyType::Width:
+			return &SetPropertyWidth;
+		case EPacketPropertyType::Height:
+			return &SetPropertyHeight;
+		case EPacketPropertyType::FingerContactConfidence:
+			return &SetPropertyFingerContactConfidence;
+		case EPacketPropertyType::DeviceContactID:
+			return &SetPropertyDeviceContactID;
+		default:
+			checkSlow(false);
+			return nullptr;
+		}
+	}
+
+	template <int32 DataBufferLength>
+	void AssignSetPropertyData(const EPacketPropertyType Type, int8 *const Data, const FWindowContext& WindowContext, const FPacketProperty& Description)
+	{
+		int8* DataCurrent = Data;
+
+		auto Add = [&DataCurrent]<typename T>(T Value)
+		{
+			*reinterpret_cast<T*>(DataCurrent) = Value;
+			DataCurrent += sizeof(T);
+		};
+
+		switch (Type)
+		{
+		case EPacketPropertyType::X:
+			Add(WindowContext.XYScale.X);
+			Add(WindowContext.XYMaximum.X);
+			Add(WindowContext.WindowSize.X);
+			break;
+		case EPacketPropertyType::Y:
+			Add(WindowContext.XYScale.Y);
+			Add(WindowContext.XYMaximum.Y);
+			Add(WindowContext.WindowSize.Y);
+			break;
+		case EPacketPropertyType::Z:
+		case EPacketPropertyType::NormalPressure:
+		case EPacketPropertyType::TangentPressure:
+		case EPacketPropertyType::ButtonPressure:
+		case EPacketPropertyType::FingerContactConfidence:
+			Add(Description.Minimum);
+			Add(Description.Maximum);
+			break;
+		case EPacketPropertyType::XTiltOrientation:
+		case EPacketPropertyType::YTiltOrientation:
+		case EPacketPropertyType::AzimuthOrientation:
+		case EPacketPropertyType::AltitudeOrientation:
+		case EPacketPropertyType::TwistOrientation:
+		case EPacketPropertyType::PitchRotation:
+		case EPacketPropertyType::RollRotation:
+		case EPacketPropertyType::YawRotation:
+		case EPacketPropertyType::Width:
+		case EPacketPropertyType::Height:
+			Add(Description.Resolution);
+			break;
+		case EPacketPropertyType::PacketStatus:
+		case EPacketPropertyType::TimerTick:
+		case EPacketPropertyType::SerialNumber:
+		case EPacketPropertyType::DeviceContactID:
+			// do nothing
+			break;
+		default:
+			checkSlow(false);
+		}
+
+		check(std::distance(Data, DataCurrent) <= DataBufferLength);
+	}
+
+	bool SetupTabletContextPacketDescriptionData(IRealTimeStylus* RealTimeStylus, const FWindowContext& WindowContext, FTabletContext& TabletContext)
+	{
+		const FWindowsStylusInputPlatformAPI& WindowsAPI = FWindowsStylusInputPlatformAPI::GetInstance();
+
+		auto InvalidatePacketDescriptions = [&PacketDescriptions = TabletContext.PacketDescriptions](const int32 StartingIndex)
+		{
+			for (int32 Index = StartingIndex; Index < static_cast<int32>(EPacketPropertyType::Num_Enumerators); ++Index)
+			{
+				PacketDescriptions[Index] = {};
+			}
+		};
+		
+		FLOAT InkToDeviceScaleX = 0.0f;
+		FLOAT InkToDeviceScaleY = 0.0f;
+		ULONG PropertiesCount = 0;
+		PACKET_PROPERTY* Properties = nullptr;
+		if (Succeeded(RealTimeStylus->GetPacketDescriptionData(TabletContext.ID, &InkToDeviceScaleX, &InkToDeviceScaleY, &PropertiesCount, &Properties), LOG_PREAMBLE))
+		{
+			if (!Properties)
+			{
+				LogWarning(LOG_PREAMBLE, FString::Format(
+					           TEXT("Retrieved nullptr when trying to get packet description data for TabletContext with ID {0}."), {TabletContext.ID}));
+				return false;
+			}
+
+			if (PropertiesCount == 0)
+			{
+				LogWarning(LOG_PREAMBLE, FString::Format(
+					           TEXT("Retrieved zero packets properties when trying to get packet description data for TabletContext with ID {0}."),
+					           {TabletContext.ID}));
+				return false;
+			}
+
+			for (ULONG Index = 0; Index < PropertiesCount; ++Index)
+			{
+#pragma warning(push)
+#pragma warning(disable : 28199)
+				const PACKET_PROPERTY& Property = Properties[Index];
+#pragma warning(pop)
+
+				bool bFoundProperty = false;
+				for (const FPacketPropertyConstant& PropertyConstant : PacketPropertyConstants)
+				{
+					if (Property.guid == PropertyConstant.Guid)
+					{
+						const ETabletSupportedProperties SupportedProperty = static_cast<ETabletSupportedProperties>(
+							1 << static_cast<std::underlying_type_t<ETabletSupportedProperties>>(PropertyConstant.PacketPropertyType)
+						);						
+						TabletContext.SupportedProperties = TabletContext.SupportedProperties | SupportedProperty;							
+
+						FPacketProperty& PropertyDescription = TabletContext.PacketDescriptions[Index];
+
+						PropertyDescription.Type = static_cast<EPacketPropertyType>(PropertyConstant.PacketPropertyType);
+						PropertyDescription.Minimum = Property.PropertyMetrics.nLogicalMin;
+						PropertyDescription.Maximum = Property.PropertyMetrics.nLogicalMax;
+						PropertyDescription.MetricUnit = static_cast<ETabletPropertyMetricUnit>(Property.PropertyMetrics.Units);
+						PropertyDescription.Resolution = Property.PropertyMetrics.fResolution;
+						PropertyDescription.SetProperty = GetSetPropertyFunc(PropertyConstant.PacketPropertyType);
+						AssignSetPropertyData<FPacketProperty::SetPropertyDataBufferLength>(PropertyConstant.PacketPropertyType, PropertyDescription.SetPropertyData, WindowContext, PropertyDescription);
+
+						if (static_cast<int8>(PropertyDescription.MetricUnit) >= static_cast<int8>(ETabletPropertyMetricUnit::Num_Enumerators))
+						{
+							PropertyDescription.MetricUnit = ETabletPropertyMetricUnit::Default;
+
+							LogWarning(LOG_PREAMBLE, FString::Format(
+									   TEXT("Encountered unknown metric unit value '{0}' while evaluating packet description data for TabletContext with ID {1}."),
+									   {static_cast<int8>(PropertyDescription.MetricUnit), TabletContext.ID}));
+						}
+
+						bFoundProperty = true;
+						break;
+					}
+				}
+
+				if (!bFoundProperty)
+				{
+					wchar_t GUIDStringBuffer[64];
+					const int32 GUIDStringLength = WindowsAPI.StringFromGUID2(Property.guid, GUIDStringBuffer, std::size(GUIDStringBuffer));
+					if (ensure(GUIDStringLength > 0))
+					{
+						LogWarning(LOG_PREAMBLE, FString::Format(
+							           TEXT("Encountered unknown property '{0}' while evaluating packet description data for TabletContext with ID {1}."),
+							           {GUIDStringBuffer, TabletContext.ID}));
+					}
+				}
+			}
+
+			WindowsAPI.CoTaskMemFree(Properties);
+
+			InvalidatePacketDescriptions(PropertiesCount);
+		}
+		else
+		{
+			LogError(LOG_PREAMBLE, FString::Format(TEXT("Could not get packet description data for TabletContext with ID {0}."), {TabletContext.ID}));
+
+			InvalidatePacketDescriptions(0);
+
+			return false;
+		}
+
+		return true;
+	}
+
+	bool SetupTabletContext(IRealTimeStylus* RealTimeStylus, const FWindowContext& WindowContext, FTabletContext& TabletContext)
+	{
+		IInkTablet* InkTablet = nullptr;
+
+		if (Failed(RealTimeStylus->GetTabletFromTabletContextId(TabletContext.ID, &InkTablet), LOG_PREAMBLE))
+		{
+			LogError(LOG_PREAMBLE, FString::Format(TEXT("Could not get tablet context data for ID {0}."), {TabletContext.ID}));
+			return false;
+		}
+
+		bool bSuccess = true;
+
+		bSuccess &= SetupTabletContextMetadata(*InkTablet, TabletContext);
+		bSuccess &= SetupTabletContextPacketDescriptionData(RealTimeStylus, WindowContext, TabletContext);
+
+		return bSuccess;
+	}
+	
+	bool FWindowsStylusInputPluginBase::AddEventHandler(IStylusInputEventHandler* EventHandler)
+	{
+		check(EventHandler);
+
+		if (EventHandlers.Contains(EventHandler))
+		{
+			LogWarning(LOG_PREAMBLE, FString::Format(TEXT("Event handler '{0}' already exists in {1} plugin."), {EventHandler->GetName(), GetName()}));
+			return false;
+		}
+		
+		EventHandlers.Add(EventHandler);
+
+		LogVerbose(LOG_PREAMBLE, FString::Format(TEXT("Event handler '{0}' was added to {1} plugin."), {EventHandler->GetName(), GetName()}));
+
+		return true;
+	}
+
+	
+	bool FWindowsStylusInputPluginBase::RemoveEventHandler(IStylusInputEventHandler* EventHandler)
+	{
+		check(EventHandler);
+
+		const bool bWasRemoved = EventHandlers.Remove(EventHandler) > 0;
+
+		if (bWasRemoved)
+		{
+			LogVerbose(LOG_PREAMBLE, FString::Format(TEXT("Event handler '{0}' was removed from {1} plugin."), {EventHandler->GetName(), GetName()}));
+		}
+		
+		return bWasRemoved;
+	}
+
+	FWindowsStylusInputPluginBase::FWindowsStylusInputPluginBase(FGetWindowContextCallback&& GetWindowContextCallback,
+	                                                   FUpdateTabletContextsCallback&& UpdateTabletContextsCallback)
+		: GetWindowContextCallback(MoveTemp(GetWindowContextCallback))
+		, UpdateTabletContextsCallback(MoveTemp(UpdateTabletContextsCallback)) 
+	{
+		check(this->GetWindowContextCallback.IsBound());
+		check(this->UpdateTabletContextsCallback.IsBound());
+	}
+
+	void FWindowsStylusInputPluginBase::DebugEvent(const FString& Message) const
+	{
+		for (IStylusInputEventHandler* EventHandler : EventHandlers)
+		{
+			EventHandler->OnDebugEvent(Message);
+		}
+	}
+
+	HRESULT FWindowsStylusInputPluginBase::ProcessRealTimeStylusEnabled(IRealTimeStylus* RealTimeStylus, const uint32 TabletContextIDsCount,
+	                                                               const TABLET_CONTEXT_ID* TabletContextIDs)
+	{
+		return UpdateTabletContexts(RealTimeStylus, TabletContextIDsCount, TabletContextIDs);
+	}
+
+	HRESULT FWindowsStylusInputPluginBase::ProcessRealTimeStylusDisabled(IRealTimeStylus* RealTimeStylus, const uint32 TabletContextIDsCount,
+	                                                                const TABLET_CONTEXT_ID* TabletContextIDs)
+	{
+		return UpdateTabletContexts(RealTimeStylus, TabletContextIDsCount, TabletContextIDs);
+	}
+
+	HRESULT FWindowsStylusInputPluginBase::ProcessPackets(const StylusInfo* StylusInfo, uint32 PacketCount, uint32 PacketBufferLength, EPacketType Type,
+	                                                 const int32* PacketBuffer)
+	{
+		checkSlow(StylusInfo);
+		checkSlow(PacketBuffer);
+
+		const FPacketProperty* Descriptions = GetPacketDescriptions(StylusInfo->tcid);
+		if (!Descriptions)
+		{
+			return E_FAIL;
+		}
+
+		const uint32 PropertyCount = PacketBufferLength / PacketCount;
+		checkSlow(PropertyCount < static_cast<uint32>(EPacketPropertyType::Num_Enumerators));
+
+		for (uint32 PacketIndex = 0; PacketIndex < PacketCount; PacketIndex += PropertyCount)
+		{
+			PacketStats.NewPacket();
+
+			FStylusInputPacket Packet{StylusInfo->tcid, StylusInfo->cid, Type};
+			
+			for (uint32 PropertyIndex = 0; PropertyIndex < PropertyCount; ++PropertyIndex)
+			{
+				const FPacketProperty& Description = Descriptions[PropertyIndex];
+				checkSlow(Description.SetProperty);
+
+				Description.SetProperty(Packet, PacketBuffer[PropertyIndex], Description.SetPropertyData);
+			}
+
+			for (IStylusInputEventHandler* EventHandler : EventHandlers)
+			{
+				EventHandler->OnPacket(Packet);
+			}
+
+			PacketBuffer += PropertyCount;
+		}
+	
+		return S_OK;
+	}
+
+	FString FWindowsStylusInputPluginBase::ProcessError(RealTimeStylusDataInterest DataInterest, const HRESULT ErrorCode)
+	{
+		const char* DataInterestStr = [DataInterest]
+		{
+			switch (DataInterest)
+			{
+			case RTSDI_AllData:
+				return "AllData";
+			case RTSDI_None:
+				return "None";
+			case RTSDI_Error:
+				return "Error";
+			case RTSDI_RealTimeStylusEnabled:
+				return "RealTimeStylusEnabled";
+			case RTSDI_RealTimeStylusDisabled:
+				return "RealTimeStylusDisabled";
+			case RTSDI_StylusNew:
+				return "StylusNew";
+			case RTSDI_StylusInRange:
+				return "StylusInRange";
+			case RTSDI_InAirPackets:
+				return "InAirPackets";
+			case RTSDI_StylusOutOfRange:
+				return "StylusOutOfRange";
+			case RTSDI_StylusDown:
+				return "StylusDown";
+			case RTSDI_Packets:
+				return "Packets";
+			case RTSDI_StylusUp:
+				return "StylusUp";
+			case RTSDI_StylusButtonUp:
+				return "StylusButtonUp";
+			case RTSDI_StylusButtonDown:
+				return "StylusButtonDown";
+			case RTSDI_SystemEvents:
+				return "SystemEvents";
+			case RTSDI_TabletAdded:
+				return "TabletAdded";
+			case RTSDI_TabletRemoved:
+				return "TabletRemoved";
+			case RTSDI_CustomStylusDataAdded:
+				return "CustomStylusDataAdded";
+			case RTSDI_UpdateMapping:
+				return "UpdateMapping";
+			case RTSDI_DefaultEvents:
+				return "DefaultEvents";
+			}
+			return "<unknown>";
+		}();
+
+		return FString::Format(TEXT("Error in {0} plugin: {1}, Error={2} ({3})"), {
+			                       GetName(), DataInterestStr, _com_error(ErrorCode).ErrorMessage(), static_cast<int32>(ErrorCode)
+		                       });
+	}
+
+	HRESULT FWindowsStylusInputPluginBase::ProcessTabletAdded(IInkTablet* Tablet)
+	{
+		return E_NOTIMPL;
+	}
+
+	HRESULT FWindowsStylusInputPluginBase::ProcessTabletRemoved(LONG TabletIndex)
+	{
+		return E_NOTIMPL;
+	}
+
+	const FPacketProperty* FWindowsStylusInputPluginBase::GetPacketDescriptions(uint32 TabletContextId) const
+	{
+		const TSharedPtr<FTabletContext> TabletContext = TabletContexts.Get(TabletContextId);
+		return TabletContext ? TabletContext->PacketDescriptions : nullptr;
+	}
+
+	HRESULT FWindowsStylusInputPluginBase::UpdateTabletContexts(IRealTimeStylus* RealTimeStylus, uint32 TabletContextIDsNum, const TABLET_CONTEXT_ID* TabletContextIDs)
+	{
+		const FWindowContext& WindowContext = GetWindowContextCallback.Execute();
+		
+		bool bSuccess = true;
+
+		// Remove outdated tablet contexts
+		uint32 TabletContextIndex = 0;
+		while (TabletContextIndex < TabletContexts.Num())
+		{
+			const uint32 ExistingTabletContextID = TabletContexts[TabletContextIndex]->ID;
+
+			const bool bContainsID = [ExistingTabletContextID, TabletContextIDsNum, TabletContextIDs]
+			{
+				for (uint32 TabletContextIDsIndex = 0; TabletContextIDsIndex < TabletContextIDsNum; ++TabletContextIDsIndex)
+				{
+					if (ExistingTabletContextID == TabletContextIDs[TabletContextIDsIndex])
+					{
+						return true;
+					}
+				}
+				return false;
+			}();
+
+			if (bContainsID)
+			{
+				++TabletContextIndex;
+			}
+			else
+			{
+				TabletContexts.Remove(ExistingTabletContextID);
+
+				LogVerbose(LOG_PREAMBLE, FString::Format(TEXT("Removed tablet context data for ID {0}."), {ExistingTabletContextID}));
+			}
+		}
+
+		// Add new tablet contexts
+		for (uint32 TabletContextIDsIndex = 0; TabletContextIDsIndex < TabletContextIDsNum; ++TabletContextIDsIndex)
+		{
+			const uint32 TabletContextID = TabletContextIDs[TabletContextIDsIndex];
+			if (!TabletContexts.Contains(TabletContextID))
+			{
+				const TSharedRef<FTabletContext> TabletContextRef = TabletContexts.Add(TabletContextID);
+				FTabletContext& TabletContext = TabletContextRef.Get();
+
+				bSuccess &= SetupTabletContext(RealTimeStylus, WindowContext, TabletContext);
+
+				LogVerbose(LOG_PREAMBLE, FString::Format(TEXT("Added TabletContext for ID {0} [{1}, {2}]."), {
+															TabletContext.ID, TabletContext.Name, TabletContext.PlugAndPlayID
+														}));
+			}
+		}
+
+		UpdateTabletContextsCallback.Execute(TabletContexts);
+
+		return bSuccess ? S_OK : E_FAIL;
+	}
+}
+
+#undef LOG_PREAMBLE
+
+#endif
diff --git a/Engine/Plugins/Editor/StylusInput/Source/StylusInput/Private/Windows/WindowsStylusInputPluginBase.h b/Engine/Plugins/Editor/StylusInput/Source/StylusInput/Private/Windows/WindowsStylusInputPluginBase.h
new file mode 100644
--- /dev/null
+++ b/Engine/Plugins/Editor/StylusInput/Source/StylusInput/Private/Windows/WindowsStylusInputPluginBase.h
@@ -0,0 +1,71 @@
+﻿// Copyright Epic Games, Inc. All Rights Reserved.
+
+#pragma once
+
+#if PLATFORM_WINDOWS
+
+#include <StylusInputPacket.h>
+#include <Delegates/Delegate.h>
+
+#include "WindowsStylusInputStats.h"
+#include "WindowsStylusInputTabletContext.h"
+
+#include <Windows/AllowWindowsPlatformTypes.h>
+	#include <Microsoft/COMPointer.h>
+	#include <RTSCom.h>
+#include <Windows/HideWindowsPlatformTypes.h>
+
+namespace UE::StylusInput
+{
+	class IStylusInputEventHandler;
+}
+
+namespace UE::StylusInput::Private::Windows
+{
+	class FWindowsStylusInputPlatformAPI;
+
+	DECLARE_DELEGATE_RetVal(const FWindowContext&, FGetWindowContextCallback);
+	DECLARE_DELEGATE_OneParam(FUpdateTabletContextsCallback, const FTabletContextContainer&);
+
+	class FWindowsStylusInputPluginBase
+	{
+	public:
+		bool AddEventHandler(IStylusInputEventHandler* EventHandler);
+		bool RemoveEventHandler(IStylusInputEventHandler* EventHandler);
+		int32 NumEventHandlers() const { return EventHandlers.Num(); }
+
+		float GetPacketsPerSecond() const { return PacketStats.GetPacketsPerSecond(); }
+
+	protected:
+		FWindowsStylusInputPluginBase(FGetWindowContextCallback&& GetWindowContextCallback, FUpdateTabletContextsCallback&& UpdateTabletContextsCallback);
+		~FWindowsStylusInputPluginBase() = default;
+
+		virtual FString GetName() const = 0;
+
+		void DebugEvent(const FString& Message) const;
+
+		HRESULT ProcessRealTimeStylusEnabled(IRealTimeStylus* RealTimeStylus, uint32 TabletContextIDsCount, const TABLET_CONTEXT_ID* TabletContextIDs);
+		HRESULT ProcessRealTimeStylusDisabled(IRealTimeStylus* RealTimeStylus, uint32 TabletContextIDsCount, const TABLET_CONTEXT_ID* TabletContextIDs);
+		HRESULT ProcessPackets(const StylusInfo* StylusInfo, uint32 PacketCount, uint32 PacketBufferLength, EPacketType Type, const int32* PacketBuffer);
+		FString ProcessError(RealTimeStylusDataInterest DataInterest, HRESULT ErrorCode);
+		HRESULT ProcessTabletAdded(IInkTablet* Tablet);
+		HRESULT ProcessTabletRemoved(LONG TabletIndex);
+
+	private:
+		const FPacketProperty* GetPacketDescriptions(uint32 TabletContextId) const;
+		HRESULT UpdateTabletContexts(IRealTimeStylus* RealTimeStylus, uint32 TabletContextIDsNum, const TABLET_CONTEXT_ID* TabletContextIDs);
+
+		FPacketStats PacketStats;
+
+		FGetWindowContextCallback GetWindowContextCallback;
+		FUpdateTabletContextsCallback UpdateTabletContextsCallback;
+
+		FTabletContextContainer TabletContexts;
+
+		TArray<IStylusInputEventHandler*> EventHandlers;
+	};
+}
+
+#define STYLUSINPUT_SHOW_NOTIMPL_ERRORS 0
+
+#endif
diff --git a/Engine/Plugins/Editor/StylusInput/Source/StylusInput/Private/Windows/WindowsStylusInputPluginSync.cpp b/Engine/Plugins/Editor/StylusInput/Source/StylusInput/Private/Windows/WindowsStylusInputPluginSync.cpp
new file mode 100644
--- /dev/null
+++ b/Engine/Plugins/Editor/StylusInput/Source/StylusInput/Private/Windows/WindowsStylusInputPluginSync.cpp
@@ -0,0 +1,128 @@
+﻿// Copyright Epic Games, Inc. All Rights Reserved.
+
+#pragma once
+
+#if PLATFORM_WINDOWS
+
+#include "WindowsStylusInputPluginSync.h"
+
+#include "WindowsStylusInputPlatformAPI.h"
+
+namespace UE::StylusInput::Private::Windows
+{
+	FWindowsStylusInputPluginSync::FWindowsStylusInputPluginSync(FGetWindowContextCallback&& GetWindowContext,
+	                                                   FUpdateTabletContextsCallback&& UpdateTabletContextsCallback)
+		: FWindowsStylusInputPluginBase(MoveTemp(GetWindowContext), MoveTemp(UpdateTabletContextsCallback))
+	{
+	}
+
+	FWindowsStylusInputPluginSync::~FWindowsStylusInputPluginSync()
+	{
+		if (FreeThreadedMarshaler)
+		{
+			FreeThreadedMarshaler->Release();
+		}
+	}
+
+	HRESULT FWindowsStylusInputPluginSync::CreateFreeThreadMarshaler()
+	{
+		check(FreeThreadedMarshaler == nullptr);
+
+		const FWindowsStylusInputPlatformAPI& WindowsAPI = FWindowsStylusInputPlatformAPI::GetInstance();
+		return WindowsAPI.CoCreateFreeThreadedMarshaler(this, &FreeThreadedMarshaler);
+	}
+
+	HRESULT FWindowsStylusInputPluginSync::RealTimeStylusEnabled(IRealTimeStylus* RealTimeStylus, ULONG TabletContextIDsCount, const TABLET_CONTEXT_ID* TabletContextIDs)
+	{
+		return ProcessRealTimeStylusEnabled(RealTimeStylus, TabletContextIDsCount, TabletContextIDs);
+	}
+
+	HRESULT FWindowsStylusInputPluginSync::RealTimeStylusDisabled(IRealTimeStylus* RealTimeStylus, ULONG TabletContextIDsCount, const TABLET_CONTEXT_ID* TabletContextIDs)
+	{
+		return ProcessRealTimeStylusDisabled(RealTimeStylus, TabletContextIDsCount, TabletContextIDs);
+	}
+
+	HRESULT FWindowsStylusInputPluginSync::StylusDown(IRealTimeStylus* RealTimeStylus, const StylusInfo* StylusInfo, ULONG PropertyCount, LONG* PacketBuffer, LONG**)
+	{
+		return ProcessPackets(StylusInfo, 1, PropertyCount, EPacketType::StylusDown, reinterpret_cast<int32*>(PacketBuffer));
+	}
+
+	HRESULT FWindowsStylusInputPluginSync::StylusUp(IRealTimeStylus* RealTimeStylus, const StylusInfo* StylusInfo, ULONG PropertyCount, LONG* PacketBuffer, LONG**)
+	{
+		return ProcessPackets(StylusInfo, 1, PropertyCount, EPacketType::StylusUp, reinterpret_cast<int32*>(PacketBuffer));
+	}
+
+	HRESULT FWindowsStylusInputPluginSync::TabletAdded(IRealTimeStylus* RealTimeStylus, IInkTablet* Tablet)
+	{
+		return ProcessTabletAdded(Tablet);
+	}
+
+	HRESULT FWindowsStylusInputPluginSync::TabletRemoved(IRealTimeStylus* RealTimeStylus, LONG TabletIndex)
+	{
+		return ProcessTabletRemoved(TabletIndex);
+	}
+
+	HRESULT FWindowsStylusInputPluginSync::InAirPackets(IRealTimeStylus* RealTimeStylus, const StylusInfo* StylusInfo, ULONG PacketCount, ULONG PacketBufferLength,
+													LONG* PacketBuffer, ULONG*, LONG**)
+	{
+		return ProcessPackets(StylusInfo, PacketCount, PacketBufferLength, EPacketType::AboveDigitizer, reinterpret_cast<int32*>(PacketBuffer));
+	}
+
+	HRESULT FWindowsStylusInputPluginSync::Packets(IRealTimeStylus* RealTimeStylus, const StylusInfo* StylusInfo, const ULONG PacketCount,
+	                                          const ULONG PacketBufferLength, LONG* PacketBuffer, ULONG*, LONG**)
+	{
+		return ProcessPackets(StylusInfo, PacketCount, PacketBufferLength, EPacketType::OnDigitizer, reinterpret_cast<int32*>(PacketBuffer));
+	}
+
+	HRESULT FWindowsStylusInputPluginSync::StylusInRange(IRealTimeStylus* RealTimeStylus, TABLET_CONTEXT_ID TabletContextID, STYLUS_ID StylusID)
+	{
+		return E_NOTIMPL;
+	}
+
+	HRESULT FWindowsStylusInputPluginSync::StylusOutOfRange(IRealTimeStylus* RealTimeStylus, TABLET_CONTEXT_ID TabletContextID, STYLUS_ID StylusID)
+	{
+		return E_NOTIMPL;
+	}
+
+	HRESULT FWindowsStylusInputPluginSync::StylusButtonDown(IRealTimeStylus* piRtsSrc, STYLUS_ID sid, const GUID* pGuidStylusButton, POINT* pStylusPos)
+	{
+		return E_NOTIMPL;
+	}
+
+	HRESULT FWindowsStylusInputPluginSync::StylusButtonUp(IRealTimeStylus* piRtsSrc, STYLUS_ID sid, const GUID* pGuidStylusButton, POINT* pStylusPos)
+	{
+		return E_NOTIMPL;
+	}
+
+	HRESULT FWindowsStylusInputPluginSync::CustomStylusDataAdded(IRealTimeStylus* piRtsSrc, const GUID* pGuidId, ULONG cbData, const BYTE* pbData)
+	{
+		return E_NOTIMPL;
+	}
+
+	HRESULT FWindowsStylusInputPluginSync::SystemEvent(IRealTimeStylus* piRtsSrc, TABLET_CONTEXT_ID tcid, STYLUS_ID sid, SYSTEM_EVENT event, SYSTEM_EVENT_DATA eventdata)
+	{
+		return E_NOTIMPL;
+	}
+
+	HRESULT FWindowsStylusInputPluginSync::Error(IRealTimeStylus*, IStylusPlugin* Plugin, const RealTimeStylusDataInterest DataInterest, const HRESULT ErrorCode, LONG_PTR*)
+	{
+#if STYLUSINPUT_SHOW_NOTIMPL_ERRORS
+		const FString ErrorMessage = ProcessError(DataInterest, ErrorCode);
+		DebugEvent(ErrorMessage);
+#endif
+
+		return S_OK;
+	}
+
+	HRESULT FWindowsStylusInputPluginSync::UpdateMapping(IRealTimeStylus* piRtsSrc)
+	{
+		return E_NOTIMPL;
+	}
+
+	HRESULT FWindowsStylusInputPluginSync::DataInterest(RealTimeStylusDataInterest* pDataInterest)
+	{
+		*pDataInterest = RTSDI_AllData;
+		return S_OK;
+	}
+}
+#endif
diff --git a/Engine/Plugins/Editor/StylusInput/Source/StylusInput/Private/Windows/WindowsStylusInputPluginSync.h b/Engine/Plugins/Editor/StylusInput/Source/StylusInput/Private/Windows/WindowsStylusInputPluginSync.h
new file mode 100644
--- /dev/null
+++ b/Engine/Plugins/Editor/StylusInput/Source/StylusInput/Private/Windows/WindowsStylusInputPluginSync.h
@@ -0,0 +1,82 @@
+﻿// Copyright Epic Games, Inc. All Rights Reserved.
+
+#pragma once
+
+#if PLATFORM_WINDOWS
+
+#include "WindowsStylusInputPluginBase.h"
+
+namespace UE::StylusInput::Private::Windows
+{
+	class FWindowsStylusInputPluginSync final : public IStylusSyncPlugin, public FWindowsStylusInputPluginBase
+	{
+	public:
+		explicit FWindowsStylusInputPluginSync(FGetWindowContextCallback&& GetWindowContext, FUpdateTabletContextsCallback&& UpdateTabletContextsCallback);
+		~FWindowsStylusInputPluginSync();
+		
+		HRESULT CreateFreeThreadMarshaler();
+
+		// IUnknown
+		virtual HRESULT QueryInterface(REFIID InterfaceID, void** InterfaceObject) override
+		{
+			if (InterfaceID == IID_IStylusSyncPlugin || InterfaceID == IID_IUnknown)
+			{
+				*InterfaceObject = this;
+				AddRef();
+				return S_OK;
+			}
+
+			if (InterfaceID == IID_IMarshal && FreeThreadedMarshaler)
+			{
+				return FreeThreadedMarshaler->QueryInterface(InterfaceID, InterfaceObject);
+			}
+
+			*InterfaceObject = nullptr;
+			return E_NOINTERFACE;
+		}
+
+		virtual ULONG AddRef() override
+		{
+			return ++RefCount;
+		}
+
+		virtual ULONG Release() override
+		{
+			const int32 NewRefCount = --RefCount;
+			if (NewRefCount == 0)
+				delete this;
+
+			return NewRefCount;
+		}
+
+		// IStylusPlugin
+		virtual HRESULT RealTimeStylusEnabled(IRealTimeStylus* RealTimeStylus, ULONG TabletContextIDsCount, const TABLET_CONTEXT_ID* TabletContextIDs) override;
+		virtual HRESULT RealTimeStylusDisabled(IRealTimeStylus* RealTimeStylus, ULONG TabletContextIDsCount, const TABLET_CONTEXT_ID* TabletContextIDs) override;
+		virtual HRESULT StylusDown(IRealTimeStylus* RealTimeStylus, const StylusInfo* StylusInfo, ULONG PropertyCount, LONG* PacketBuffer, LONG** ppInOutPkt) override;
+		virtual HRESULT StylusUp(IRealTimeStylus* piRtsSrc, const StylusInfo* pStylusInfo, ULONG cPropCountPerPkt, LONG* pPacket, LONG** ppInOutPkt) override;
+		virtual HRESULT TabletAdded(IRealTimeStylus* RealTimeStylus, IInkTablet* Tablet) override;
+		virtual HRESULT TabletRemoved(IRealTimeStylus* RealTimeStylus, LONG TabletIndex) override;
+		virtual HRESULT InAirPackets(IRealTimeStylus* RealTimeStylus, const StylusInfo* StylusInfo, ULONG PacketCount, ULONG PacketBufferLength, LONG* PacketBuffer, ULONG* InOutPacketCount, LONG** InOutPacketBuffer) override;
+		virtual HRESULT Packets(IRealTimeStylus* RealTimeStylus, const StylusInfo* StylusInfo, ULONG PacketCount, ULONG PacketBufferLength, LONG* PacketBuffer, ULONG* InOutPacketCount, LONG** InOutPacketBuffer) override;
+		virtual HRESULT StylusInRange(IRealTimeStylus* piRtsSrc, TABLET_CONTEXT_ID tcid, STYLUS_ID sid) override;
+		virtual HRESULT StylusOutOfRange(IRealTimeStylus* piRtsSrc, TABLET_CONTEXT_ID tcid, STYLUS_ID sid) override;
+		virtual HRESULT StylusButtonDown(IRealTimeStylus* piRtsSrc, STYLUS_ID sid, const GUID* pGuidStylusButton, POINT* pStylusPos) override;
+		virtual HRESULT StylusButtonUp(IRealTimeStylus* piRtsSrc, STYLUS_ID sid, const GUID* pGuidStylusButton, POINT* pStylusPos) override;
+		virtual HRESULT CustomStylusDataAdded(IRealTimeStylus* piRtsSrc, const GUID* pGuidId, ULONG cbData, const BYTE* pbData) override;
+		virtual HRESULT SystemEvent(IRealTimeStylus* piRtsSrc, TABLET_CONTEXT_ID tcid, STYLUS_ID sid, SYSTEM_EVENT event, SYSTEM_EVENT_DATA eventdata) override;
+		virtual HRESULT Error(IRealTimeStylus* RealTimeStylus, IStylusPlugin* Plugin, RealTimeStylusDataInterest DataInterest, HRESULT ErrorCode, LONG_PTR* InternalKey) override;
+		virtual HRESULT UpdateMapping(IRealTimeStylus* piRtsSrc) override;
+		virtual HRESULT DataInterest(RealTimeStylusDataInterest* pDataInterest) override;
+
+	protected:
+		// FWindowsStylusInputPluginBase
+		virtual FString GetName() const override { return "Asynchronous"; }
+		
+	private:
+		int32 RefCount = 1;
+
+		IUnknown* FreeThreadedMarshaler = nullptr;
+	};
+}
+
+#endif
diff --git a/Engine/Plugins/Editor/StylusInput/Source/StylusInput/Private/Windows/WindowsStylusInputStats.h b/Engine/Plugins/Editor/StylusInput/Source/StylusInput/Private/Windows/WindowsStylusInputStats.h
new file mode 100644
--- /dev/null
+++ b/Engine/Plugins/Editor/StylusInput/Source/StylusInput/Private/Windows/WindowsStylusInputStats.h
@@ -0,0 +1,45 @@
+﻿// Copyright Epic Games, Inc. All Rights Reserved.
+
+#pragma once
+
+#if PLATFORM_WINDOWS
+
+// #include <HAL/Platform.h>
+#include "Windows/WindowsPlatformTime.h"
+
+namespace UE::StylusInput::Private::Windows
+{
+	class FPacketStats
+	{
+	public:
+		void NewPacket()
+		{
+			const uint32 Timestamp = FPlatformTime::Cycles();
+
+			if (Latest - Earliest >= CyclesPerSecond)
+			{
+				PacketsPerSecond = Earliest != Latest ? NumPackets / ((Latest - Earliest) / CyclesPerSecond) : 0.0f;
+				Earliest = Timestamp;
+				NumPackets = -1;
+			}
+
+			Latest = Timestamp;
+			++NumPackets;
+		}
+
+		float GetPacketsPerSecond() const
+		{
+			const uint32 Timestamp = FPlatformTime::Cycles();
+			return Timestamp - Latest <= CyclesPerSecond ? PacketsPerSecond : 0.0f;
+		}
+
+	private:
+		uint32 Earliest = 0;
+		uint32 Latest = 0;
+		uint32 NumPackets = 0;
+		float PacketsPerSecond = 0.0f;
+		const double CyclesPerSecond = 1.0 / FPlatformTime::GetSecondsPerCycle();
+	};
+}
+
+#endif
diff --git a/Engine/Plugins/Editor/StylusInput/Source/StylusInput/Private/Windows/WindowsStylusInputTabletContext.cpp b/Engine/Plugins/Editor/StylusInput/Source/StylusInput/Private/Windows/WindowsStylusInputTabletContext.cpp
new file mode 100644
--- /dev/null
+++ b/Engine/Plugins/Editor/StylusInput/Source/StylusInput/Private/Windows/WindowsStylusInputTabletContext.cpp
@@ -0,0 +1,40 @@
+﻿// Copyright Epic Games, Inc. All Rights Reserved.
+
+#pragma once
+
+#if PLATFORM_WINDOWS
+
+#include "WindowsStylusInputTabletContext.h"
+
+#include <Microsoft/COMPointer.h>
+#include <msinkaut.h>
+
+namespace UE::StylusInput::Private::Windows
+{
+	FPacketPropertyConstant PacketPropertyConstants[] = {
+		{EPacketPropertyType::X,                       GUID_PACKETPROPERTY_GUID_X,                        STR_GUID_X                      },
+		{EPacketPropertyType::Y,                       GUID_PACKETPROPERTY_GUID_Y,                        STR_GUID_Y                      },
+		{EPacketPropertyType::Z,                       GUID_PACKETPROPERTY_GUID_Z,                        STR_GUID_Z                      },
+		{EPacketPropertyType::PacketStatus,            GUID_PACKETPROPERTY_GUID_PACKET_STATUS,            STR_GUID_PAKETSTATUS            },
+		{EPacketPropertyType::TimerTick,               GUID_PACKETPROPERTY_GUID_TIMER_TICK,               STR_GUID_TIMERTICK              },
+		{EPacketPropertyType::SerialNumber,            GUID_PACKETPROPERTY_GUID_SERIAL_NUMBER,            STR_GUID_SERIALNUMBER           },
+		{EPacketPropertyType::NormalPressure,          GUID_PACKETPROPERTY_GUID_NORMAL_PRESSURE,          STR_GUID_NORMALPRESSURE         },
+		{EPacketPropertyType::TangentPressure,         GUID_PACKETPROPERTY_GUID_TANGENT_PRESSURE,         STR_GUID_TANGENTPRESSURE        },
+		{EPacketPropertyType::ButtonPressure,          GUID_PACKETPROPERTY_GUID_BUTTON_PRESSURE,          STR_GUID_BUTTONPRESSURE         },
+		{EPacketPropertyType::XTiltOrientation,        GUID_PACKETPROPERTY_GUID_X_TILT_ORIENTATION,       STR_GUID_XTILTORIENTATION       },
+		{EPacketPropertyType::YTiltOrientation,        GUID_PACKETPROPERTY_GUID_Y_TILT_ORIENTATION,       STR_GUID_YTILTORIENTATION       },
+		{EPacketPropertyType::AzimuthOrientation,      GUID_PACKETPROPERTY_GUID_AZIMUTH_ORIENTATION,      STR_GUID_AZIMUTHORIENTATION     },
+		{EPacketPropertyType::AltitudeOrientation,     GUID_PACKETPROPERTY_GUID_ALTITUDE_ORIENTATION,     STR_GUID_ALTITUDEORIENTATION    },
+		{EPacketPropertyType::TwistOrientation,        GUID_PACKETPROPERTY_GUID_TWIST_ORIENTATION,        STR_GUID_TWISTORIENTATION       },
+		{EPacketPropertyType::PitchRotation,           GUID_PACKETPROPERTY_GUID_PITCH_ROTATION,           STR_GUID_PITCHROTATION          },
+		{EPacketPropertyType::RollRotation,            GUID_PACKETPROPERTY_GUID_ROLL_ROTATION,            STR_GUID_ROLLROTATION           },
+		{EPacketPropertyType::YawRotation,             GUID_PACKETPROPERTY_GUID_YAW_ROTATION,             STR_GUID_YAWROTATION            },
+		{EPacketPropertyType::Width,                   GUID_PACKETPROPERTY_GUID_WIDTH,                    STR_GUID_WIDTH                  },
+		{EPacketPropertyType::Height,                  GUID_PACKETPROPERTY_GUID_HEIGHT,                   STR_GUID_HEIGHT                 },
+		{EPacketPropertyType::FingerContactConfidence, GUID_PACKETPROPERTY_GUID_FINGERCONTACTCONFIDENCE,  STR_GUID_FINGERCONTACTCONFIDENCE},
+		{EPacketPropertyType::DeviceContactID,         GUID_PACKETPROPERTY_GUID_DEVICE_CONTACT_ID,        STR_GUID_DEVICE_CONTACT_ID      }
+	};
+	static_assert(std::size(PacketPropertyConstants) == static_cast<uint32>(EPacketPropertyType::Num_Enumerators));
+}
+
+#endif
\ No newline at end of file
diff --git a/Engine/Plugins/Editor/StylusInput/Source/StylusInput/Private/Windows/WindowsStylusInputTabletContext.h b/Engine/Plugins/Editor/StylusInput/Source/StylusInput/Private/Windows/WindowsStylusInputTabletContext.h
new file mode 100644
--- /dev/null
+++ b/Engine/Plugins/Editor/StylusInput/Source/StylusInput/Private/Windows/WindowsStylusInputTabletContext.h
@@ -0,0 +1,153 @@
+﻿// Copyright Epic Games, Inc. All Rights Reserved.
+
+#pragma once
+
+#if PLATFORM_WINDOWS
+
+#include <StylusInputPacket.h>
+#include <StylusInputTabletContext.h>
+#include <Math/IntRect.h>
+#include <Math/Vector.h>
+
+#include "WindowsStylusInputUtils.h"
+
+// The following includes are needed for 'GUID' type.
+#include <Windows/AllowWindowsPlatformTypes.h>
+#include <Windows/HideWindowsPlatformTypes.h>
+
+namespace UE::StylusInput::Private::Windows
+{
+	enum class EPacketPropertyType : int8
+	{
+		X = 0, // The x-coordinate in the tablet coordinate space. Each packet contains this property by default. The origin (0,0) of the tablet is the upper-left corner.
+		Y, // The y-coordinate in the tablet coordinate space. Each packet contains this property by default. The origin (0,0) of the tablet is the upper-left corner.
+		Z, // The z-coordinate or distance of the pen tip from the tablet surface. The TabletPropertyMetricUnit enumeration type determines the unit of measurement for this property.
+		PacketStatus, // Contains one or more of the following flag values: The cursor is touching the drawing surface (Value = 1). The cursor is inverted. For example, the eraser end of the pen is pointing toward the surface (Value = 2). Not used (Value = 4). The barrel button is pressed (Value = 8).
+		TimerTick, // The time the packet was generated.
+		SerialNumber, // The packet property for identifying the packet. This is the same value you use to retrieve the packet from the packet queue.
+		NormalPressure, // The pressure of the pen tip perpendicular to the tablet surface. The greater the pressure on the pen tip, the more ink that is drawn.
+		TangentPressure, // The pressure of the pen tip along the plane of the tablet surface.
+		ButtonPressure, // The pressure on a pressure sensitive button.
+		XTiltOrientation, // The angle between the y,z-plane and the pen and y-axis plane. Applies to a pen cursor. The value is 0 when the pen is perpendicular to the drawing surface and is positive when the pen is to the right of perpendicular.
+		YTiltOrientation, // The angle between the x,z-plane and the pen and x-axis plane. Applies to a pen cursor. The value is 0 when the pen is perpendicular to the drawing surface and is positive when the pen is upward or away from the user.
+		AzimuthOrientation, // The clockwise rotation of the cursor about the z-axis through a full circular range.
+		AltitudeOrientation, // The angle between the axis of the pen and the surface of the tablet. The value is 0 when the pen is parallel to the surface and 90 when the pen is perpendicular to the surface. The values are negative when the pen is inverted.
+		TwistOrientation, // The clockwise rotation of the cursor about its own axis.
+		PitchRotation, // The packet property that indicates whether the tip is above or below a horizontal line that is perpendicular to the writing surface. Note: This requires a 3D digitizer. The value is positive if the tip is above the line and negative if it is below the line. For example, if you hold the pen in front of you and write on an imaginary wall, the pitch is positive if the tip is above a line extending from you to the wall.
+		RollRotation, // The clockwise rotation of the pen around its own axis. Note: This requires a 3D digitizer.
+		YawRotation, // The angle of the pen to the left or right around the center of its horizontal axis when the pen is horizontal. Note: This requires a 3D digitizer. If you hold the pen in front of you and write on an imaginary wall, zero yaw indicates that the pen is perpendicular to the wall. The value is negative if the tip is to the left of perpendicular and positive if the tip is to the right of perpendicular.
+		Width, // The width of the contact area on a touch digitizer.
+		Height, // The height of the contact area on a touch digitizer.
+		FingerContactConfidence, // The level of confidence that there was finger contact on a touch digitizer.
+		DeviceContactID, // The device contact identifier for a packet.
+
+		Num_Enumerators, // THIS IS NOT A VALID ENUMERATOR, BUT IT CAN BE USED TO QUERY HOW MANY PACKET PROPERTIES THERE ARE.
+		Invalid_Enumerator = INDEX_NONE // THIS IS NOT A VALID ENUMERATOR, BUT IT CAN BE USED TO INDICATE THAT A PACKET PROPERTY ENTRY IS INVALID.
+	};
+
+	enum class ETabletPropertyMetricUnit : int8
+	{
+		Default, // Units are unknown.
+		Inches,
+		Centimeters,
+		Degrees,
+		Radians,
+		Seconds,
+		Pounds,
+		Grams,
+
+		Num_Enumerators, // THIS IS NOT A VALID ENUMERATOR, BUT IT CAN BE USED TO QUERY HOW MANY PACKET PROPERTIES THERE ARE.
+		Invalid_Enumerator = INDEX_NONE // THIS IS NOT A VALID ENUMERATOR, BUT IT CAN BE USED TO INDICATE THAT A PACKET PROPERTY ENTRY IS INVALID.
+	};
+
+	using FuncSetProperty = void(FStylusInputPacket&, int32, const int8*);
+
+	struct FPacketProperty
+	{
+		EPacketPropertyType Type = EPacketPropertyType::Invalid_Enumerator;
+		ETabletPropertyMetricUnit MetricUnit = ETabletPropertyMetricUnit::Default;
+		int32 Minimum = 0;
+		int32 Maximum = 0;
+		float Resolution = 0.0f;
+		FuncSetProperty* SetProperty = nullptr;
+
+		static constexpr int32 SetPropertyDataBufferLength = 24;
+		int8 SetPropertyData[SetPropertyDataBufferLength];
+	};
+
+	struct FPacketPropertyConstant
+	{
+		EPacketPropertyType PacketPropertyType;
+		GUID Guid;
+		const wchar_t* StrGuid;
+	};
+
+	extern FPacketPropertyConstant PacketPropertyConstants[static_cast<uint32>(EPacketPropertyType::Num_Enumerators)];
+
+	struct FWindowContext
+	{
+		FVector2d XYScale;
+		FVector2d XYMaximum;
+		FVector2d WindowSize;
+	};
+
+	class FTabletContext final : public IStylusInputTabletContext
+	{
+	public:
+		explicit FTabletContext(const uint32 ID)
+			: ID(ID)
+		{}
+
+		virtual uint32 GetID() const override { return ID; }
+		virtual FString GetName() const override { return Name; }
+		virtual FIntRect GetInputRectangle() const override { return InputRectangle; }
+		virtual ETabletHardwareCapabilities GetHardwareCapabilities() const override { return HardwareCapabilities; }
+		virtual ETabletSupportedProperties GetSupportedProperties() const override { return SupportedProperties; }
+
+		uint32 ID = 0;
+		FString Name;
+		FString PlugAndPlayID;
+		FIntRect InputRectangle;
+		ETabletSupportedProperties SupportedProperties = ETabletSupportedProperties::None;
+		ETabletHardwareCapabilities HardwareCapabilities = ETabletHardwareCapabilities::None;
+		FPacketProperty PacketDescriptions[static_cast<uint32>(EPacketPropertyType::Num_Enumerators)];
+	};
+
+	using FTabletContextContainer = TSharedRefDataContainer<FTabletContext, false>;
+	using FTabletContextThreadSafeContainer = TSharedRefDataContainer<FTabletContext, true>;
+
+	class FStylusButton final : IStylusInputStylusButton
+	{
+	public:
+		virtual FString GetID() const override { return ID; }
+		virtual FString GetName() const override { return Name; }
+
+		FString ID;
+		FString Name;
+	};
+
+	class FStylusInfo final : public IStylusInputStylusInfo
+	{
+	public:
+		explicit FStylusInfo(const uint32 ID)
+			: ID(ID)
+		{}
+
+		virtual uint32 GetID() const override { return ID; }
+		virtual FString GetName() const override { return Name; }
+		virtual uint32 GetNumButtons() const override { return Buttons.Num(); }
+
+		virtual const IStylusInputStylusButton* GetButton(int32 Index) const override
+		{
+			return 0 <= Index && Index < Buttons.Num() ? reinterpret_cast<const IStylusInputStylusButton*>(&Buttons[Index]) : nullptr;
+		}
+
+		uint32 ID = 0;
+		FString Name;
+		TArray<FStylusButton> Buttons;
+	};
+
+	using FStylusInfoThreadSafeContainer = TSharedRefDataContainer<FStylusInfo, true>;
+}
+
+#endif
diff --git a/Engine/Plugins/Editor/StylusInput/Source/StylusInput/Private/Windows/WindowsStylusInputUtils.cpp b/Engine/Plugins/Editor/StylusInput/Source/StylusInput/Private/Windows/WindowsStylusInputUtils.cpp
new file mode 100644
--- /dev/null
+++ b/Engine/Plugins/Editor/StylusInput/Source/StylusInput/Private/Windows/WindowsStylusInputUtils.cpp
@@ -0,0 +1,28 @@
+﻿// Copyright Epic Games, Inc. All Rights Reserved.
+
+#if PLATFORM_WINDOWS
+
+#include "WindowsStylusInputUtils.h"
+
+#include <Windows/AllowWindowsPlatformAtomics.h>
+	#include <comdef.h>
+#include <Windows/HideWindowsPlatformAtomics.h>
+
+#include <StylusInputUtils.h>
+
+namespace UE::StylusInput::Private::Windows
+{
+	FString GetCOMErrorMessage(const HRESULT Result)
+	{
+		const _com_error COMError(Result);
+		const LPCTSTR ErrorMessage = COMError.ErrorMessage();
+		return ErrorMessage;
+	}
+
+	void LogCOMError(const FString& Preamble, const HRESULT Result)
+	{
+		LogError(Preamble, FString::Format(TEXT("COM Interface Error: {0}"), {GetCOMErrorMessage(Result)}));
+	}
+}
+
+#endif
diff --git a/Engine/Plugins/Editor/StylusInput/Source/StylusInput/Private/Windows/WindowsStylusInputUtils.h b/Engine/Plugins/Editor/StylusInput/Source/StylusInput/Private/Windows/WindowsStylusInputUtils.h
new file mode 100644
--- /dev/null
+++ b/Engine/Plugins/Editor/StylusInput/Source/StylusInput/Private/Windows/WindowsStylusInputUtils.h
@@ -0,0 +1,115 @@
+﻿// Copyright Epic Games, Inc. All Rights Reserved.
+
+#pragma once
+
+#if PLATFORM_WINDOWS
+
+#include <Containers/UnrealString.h>
+#include <Misc/ScopeRWLock.h>
+#include <Templates/SharedPointer.h>
+
+namespace UE::StylusInput::Private::Windows
+{
+	void LogCOMError(const FString& Preamble, HRESULT Result);
+
+	inline bool Succeeded(const HRESULT Result, const FString& LogPreamble)
+	{
+		if (Result >= 0)
+		{
+			return true;
+		}
+		LogCOMError(LogPreamble, Result);
+		return false;
+	}
+
+	inline bool Failed(const HRESULT Result, const FString& LogPreamble)
+	{
+		if (Result >= 0)
+		{
+			return false;
+		}
+		LogCOMError(LogPreamble, Result);
+		return true;
+	}
+
+	template <typename DataType, bool ThreadSafe>
+	class TSharedRefDataContainer
+	{
+	public:
+		TSharedRef<DataType> Add(uint32 ID)
+		{
+			WriteScopeLockType WriteLock(RWLock);
+			return Data.Emplace_GetRef(MakeShared<DataType>(ID));
+		}
+
+		bool Contains(uint32 ID) const
+		{
+			ReadScopeLockType ReadLock(RWLock);
+			const TSharedRef<DataType>* TabletContext = Data.FindByPredicate([ID](const TSharedRef<DataType>& Context)
+			{
+				return Context->ID == ID;
+			});
+			return TabletContext != nullptr;
+		}
+
+		TSharedPtr<DataType> Get(uint32 ID) const
+		{
+			ReadScopeLockType ReadLock(RWLock);
+			const TSharedRef<DataType>* TabletContext = Data.FindByPredicate([ID](const TSharedRef<DataType>& Context)
+			{
+				return Context->ID == ID;
+			});
+			return TabletContext ? TabletContext->ToSharedPtr() : TSharedPtr<DataType>{};
+		}
+
+		bool Remove(uint32 ID)
+		{
+			int32 Index;
+			{
+				ReadScopeLockType ReadLock(RWLock);
+				Index = Data.IndexOfByPredicate([ID](const TSharedRef<DataType>& Tc) { return Tc->ID == ID; });
+				if (Index == INDEX_NONE)
+				{
+					return false;
+				}
+			}
+			WriteScopeLockType WriteLock(RWLock);
+			Data.RemoveAt(Index, EAllowShrinking::No);
+			return true;
+		}
+
+		void Clear()
+		{
+			WriteScopeLockType WriteLock(RWLock);
+			Data.Reset();
+		}
+
+		void Update(const TSharedRefDataContainer<DataType, false>& InTabletContexts)
+		{
+			WriteScopeLockType WriteLock(RWLock);
+			Data.Reset();
+			for (uint32 Index = 0, Num = InTabletContexts.Num(); Index < Num; ++Index)
+			{
+				Data.Emplace(InTabletContexts[Index]);
+			}
+		}
+
+		uint32 Num() const { return Data.Num(); }
+
+		const TSharedRef<DataType>& operator[](const uint32 Index) const { return Data[Index]; }
+
+	private:
+
+		struct FNoLock {};
+		struct FScopeNoLock { explicit FScopeNoLock(FNoLock&) {} };
+
+		using LockType = std::conditional_t<ThreadSafe, FRWLock, FNoLock>;
+		using ReadScopeLockType = std::conditional_t<ThreadSafe, FReadScopeLock, FScopeNoLock>;
+		using WriteScopeLockType = std::conditional_t<ThreadSafe, FWriteScopeLock, FScopeNoLock>;
+		
+		UE_NO_UNIQUE_ADDRESS mutable LockType RWLock;
+		TArray<TSharedRef<DataType>> Data;
+	};
+}
+
+#endif
diff --git a/Engine/Plugins/Editor/StylusInput/Source/StylusInput/Private/WindowsRealTimeStylusPlugin.cpp b/Engine/Plugins/Editor/StylusInput/Source/StylusInput/Private/WindowsRealTimeStylusPlugin.cpp
deleted file mode 100644
diff --git a/Engine/Plugins/Editor/StylusInput/Source/StylusInput/Private/WindowsRealTimeStylusPlugin.h b/Engine/Plugins/Editor/StylusInput/Source/StylusInput/Private/WindowsRealTimeStylusPlugin.h
deleted file mode 100644
diff --git a/Engine/Plugins/Editor/StylusInput/Source/StylusInput/Private/WindowsStylusInputInterface.cpp b/Engine/Plugins/Editor/StylusInput/Source/StylusInput/Private/WindowsStylusInputInterface.cpp
deleted file mode 100644
diff --git a/Engine/Plugins/Editor/StylusInput/Source/StylusInput/Private/WindowsStylusInputInterface.h b/Engine/Plugins/Editor/StylusInput/Source/StylusInput/Private/WindowsStylusInputInterface.h
deleted file mode 100644
diff --git a/Engine/Plugins/Editor/StylusInput/Source/StylusInput/Public/IStylusInputModule.h b/Engine/Plugins/Editor/StylusInput/Source/StylusInput/Public/IStylusInputModule.h
--- a/Engine/Plugins/Editor/StylusInput/Source/StylusInput/Public/IStylusInputModule.h
+++ b/Engine/Plugins/Editor/StylusInput/Source/StylusInput/Public/IStylusInputModule.h
@@ -13,7 +13,6 @@
 class FSpawnTabArgs;
 class IStylusInputDevice;
 class IStylusMessageHandler;
-class SDockTab;
 
 DEFINE_LOG_CATEGORY_STATIC(LogStylusInput, Log, All);
 
@@ -84,10 +83,4 @@ public:
 private:
 	TSharedPtr<IStylusInputInterfaceInternal> InputInterface;
 	TArray<IStylusMessageHandler*> MessageHandlers;
-
-	TSharedRef<SDockTab> OnSpawnPluginTab(const FSpawnTabArgs& Args);
 };
-
-#if UE_ENABLE_INCLUDE_ORDER_DEPRECATED_IN_5_2
-#include "Widgets/Docking/SDockTab.h"
-#endif
diff --git a/Engine/Plugins/Editor/StylusInput/Source/StylusInput/Public/StylusInput.h b/Engine/Plugins/Editor/StylusInput/Source/StylusInput/Public/StylusInput.h
new file mode 100644
--- /dev/null
+++ b/Engine/Plugins/Editor/StylusInput/Source/StylusInput/Public/StylusInput.h
@@ -0,0 +1,142 @@
+﻿// Copyright Epic Games, Inc. All Rights Reserved.
+
+#pragma once
+
+#include <Templates/SharedPointer.h>
+
+class SWindow;
+
+namespace UE::StylusInput
+{
+	class IStylusInputEventHandler;
+	class IStylusInputInstance;
+	class IStylusInputStylusInfo;
+	class IStylusInputTabletContext;
+
+	struct FStylusInputPacket;
+
+	/**
+	 * Creates a stylus input instance for a given window.
+	 * Any call to this function that returns a valid instance pointer should have a corresponding call to @see #ReleaseInstance.
+	 *
+	 * @param Window The window the stylus input is mapped to.
+	 * @returns The stylus input instance for the given window on the current platform; a nullptr is returned if the instance could not be created or the
+	 * current platform is not supported.
+	 */
+	STYLUSINPUT_API IStylusInputInstance* CreateInstance(SWindow& Window);
+
+	/**
+	 * Releases any resources held for the given stylus input instance.
+	 * This function should be called for any corresponding successful call to @see #CreateInstance, e.g. when stylus input is no longer needed or the window
+	 * related to the instance is destroyed.
+	 *
+	 * @param Instance The stylus input instance to be released.
+	 * @returns True if the given instance is not a nullptr and was successfully released. 
+	 */
+	STYLUSINPUT_API bool ReleaseInstance(IStylusInputInstance* Instance);
+
+	/**
+	 * Defines the type of thread on which an event handler will be called by the stylus input instance.
+	 */
+	enum class EEventHandlerThread : int8
+	{
+		OnGameThread,	// The event handler will be called on the game thread.
+		Asynchronous	// The event handler will be called on an unspecific thread that is not the game thread.
+	};
+
+	/**
+	 * Interface for a stylus input instance.
+	 */
+	class IStylusInputInstance
+	{
+	public:
+		virtual ~IStylusInputInstance() = default;
+
+		/**
+		 * Adds an event handler that gets called for any processed stylus input event.
+		 * Multiple event handlers can get added for a single stylus input instance. Please note that long-running calls will block any other event handlers,
+		 * the processing of additional events, and potentially the game thread itself. Thus make sure to have an event handler call return as soon as possible.
+		 *
+		 * @param EventHandler The event handler that gets called for each processed stylus input event.
+		 * @param Thread The type of thread the event handler calls are performed on.
+		 * @returns True if the event handler was added successfully.
+		 */
+		STYLUSINPUT_API virtual bool AddEventHandler(IStylusInputEventHandler* EventHandler, EEventHandlerThread Thread) = 0;
+
+		/**
+		 * Removes an event handler from this stylus input instances, and frees up all associated resources.
+		 * The given event handler should have been added previously via a call to @see AddEventHandler.
+		 *
+		 * @param EventHandler The event handler to be released.
+		 * @returns True if the event handler was removed successfully.
+		 */
+		STYLUSINPUT_API virtual bool RemoveEventHandler(IStylusInputEventHandler* EventHandler) = 0;
+
+		/**
+		 * Provides the tablet context information for a given tablet context ID.
+		 * This function should only be called with tablet context IDs that were recently provided by the stylus input instance,
+		 * e.g. within a @see FStylusPacket as part of an event handler callback.
+		 *
+		 * Calls to this function are thread-safe, i.e. can be made from an event handler callback running on any thread.
+		 *
+		 * @param TabletContextID The unique identifier of the tablet context as provided by this stylus input instance via an event handler callback. 
+		 * @returns Pointer to tablet context information or nullptr if the tablet context ID was invalid or an error occured. 
+		 */
+		STYLUSINPUT_API virtual const TSharedPtr<IStylusInputTabletContext> GetTabletContext(uint32 TabletContextID) = 0;
+
+		/**
+		 * Provides the stylus a.k.a. cursor or pen information for a given stylus ID.
+		 * This function should only be called with stylus IDs that were recently provided by the stylus input instance,
+		* e.g. within a @see FStylusPacket as part of an event handler callback.
+		 *
+		 * Calls to this function are thread-safe, i.e. can be made from an event handler callback running on any thread.
+		 *
+		 * @param StylusID The unique identifier of the stylus/cursor/pen as provided by this stylus input instance via an event handler callback. 
+		 * @returns Pointer to stylus information or nullptr if the stylus ID was invalid or an error occured. 
+		 */
+		STYLUSINPUT_API virtual const TSharedPtr<IStylusInputStylusInfo> GetStylusInfo(uint32 StylusID) = 0;
+
+		/**
+		 * Provides the approximate number of stylus input packets processed per second for the given thread type for diagnostic purposes.
+		 * The result will only be meaningful if at least one event handler was added for the given thread type.
+		 *
+		 * @param Thread Selects the set of event handlers, i.e. number of packets processed on the game thread versus processed asynchronously.
+		 * @returns Approximate number of stylus input packets processed per second, or -1 if there is no valid data available.
+		 */
+		STYLUSINPUT_API virtual float GetPacketsPerSecond(EEventHandlerThread Thread) const { return -1.0f; }
+	};
+
+	/**
+	 * Interface for a stylus input event handler.
+	 */
+	class IStylusInputEventHandler
+	{
+	public:
+		virtual ~IStylusInputEventHandler() = default;
+
+		/**
+		 * Returns the name of the event handler. This is used internally for diagnostic purposes and reporting warnings and errors.
+		 *
+		 * @returns Name of the event handler.
+		 */
+		STYLUSINPUT_API virtual FString GetName() = 0;
+
+		/**
+		 * Callback for each packet processed by the stylus input instance.
+		 * This function is called on the game thread or asynchronously to the game thread depending on how the event handler was added to the stylus input
+		 * instance. The implementation of this function is responsible for any thread synchronization that might be necessary. 
+		 *
+		 * @param Packet The packet being processed.
+		 */
+		STYLUSINPUT_API virtual void OnPacket(const FStylusInputPacket& Packet) = 0;
+
+		/**
+		 * Callback for each debug event sent by the stylus input instance for diagnostic purposes.
+		 * This function is called on the game thread or asynchronously to the game thread depending on how the event handler was added to the stylus input
+		 * instance. The implementation of this function is responsible for any thread synchronization that might be necessary. 
+		 *
+		 * @param Message The debug message sent by the stylus input instance.
+		 */
+		virtual void OnDebugEvent(const FString& Message) = 0;
+	};
+}
diff --git a/Engine/Plugins/Editor/StylusInput/Source/StylusInput/Public/StylusInputPacket.h b/Engine/Plugins/Editor/StylusInput/Source/StylusInput/Public/StylusInputPacket.h
new file mode 100644
--- /dev/null
+++ b/Engine/Plugins/Editor/StylusInput/Source/StylusInput/Public/StylusInputPacket.h
@@ -0,0 +1,92 @@
+﻿// Copyright Epic Games, Inc. All Rights Reserved.
+
+#pragma once
+
+#include <type_traits>
+#include <HAL/Platform.h>
+
+namespace UE::StylusInput
+{
+	/**
+	 * Flags for the status of the stylus pen.
+	 */
+	enum class EPenStatus : uint8
+	{
+		None				= 0b00000000,	///< There is no pen flag.
+		CursorIsTouching	= 0b00000001,	///< The pen cursor is touching the drawing surface. 
+		CursorIsInverted	= 0b00000010,	///< The pen cursor is inverted, e.g. the eraser end of the pen is pointing toward the drawing surface.
+		NotUsed				= 0b00000100,	///< Not used.
+		BarrelButtonPressed = 0b00001000	///< The barrel button is pressed.
+	};
+
+	/**
+	 * Describes the interaction between pen and drawing surface (a.k.a. digitizer) for which the packet was generated.
+	 */
+	enum class EPacketType : uint8
+	{
+		Invalid,		///< The packet is not valid.
+		OnDigitizer,	///< The packet was generated while the pen's cursor is touching the drawing surface.
+		AboveDigitizer, ///< The packet was generated while the pen's cursor is hovering in proximity above the drawing surface.
+		StylusDown,		///< The packet was generated when the pen's cursor started touching the drawing surface.
+		StylusUp		///< The packet was generated when the pen's cursor stopped touching the drawing surface.
+	};
+
+	/**
+	 * Data for an interaction between the pen and the drawing surface.
+	 *
+	 * Note that not all devices support all properties. The set of valid properties can be queried via @see IStylusInputTabletContext#GetSupportedProperties,
+	 * and a description of all properties is available in @see ETabletSupportedProperties.
+	 */
+	struct FStylusInputPacket
+	{
+		// Metadata
+		uint32 TabletContextID = 0;					///< Unique identifier for the tablet context that created the packet.
+		uint32 CursorID = 0;						///< Unique identifier for the stylus/cursor/pen that was used to create the packet.
+		EPacketType Type = EPacketType::Invalid;	///< The type of interaction between pen and drawing surface that created the packet.
+		EPenStatus PenStatus = EPenStatus::None;	///< The status of the stylus/cursor/pen when the packet was created.
+
+		// Properties
+		float X = 0.0f;
+		float Y = 0.0f;
+		float Z = 0.0f;
+		int32 TimerTick = 0;
+		int32 SerialNumber = 0;
+		float NormalPressure = 0.0f;
+		float TangentPressure = 0.0f;
+		float ButtonPressure = 0.0f;
+		float XTiltOrientation = 0.0f;
+		float YTiltOrientation = 0.0f;
+		float AzimuthOrientation = 0.0f;
+		float AltitudeOrientation = 0.0f;
+		float TwistOrientation = 0.0f;
+		float PitchRotation = 0.0f;
+		float RollRotation = 0.0f;
+		float YawRotation = 0.0f;
+		float Width = 0.0f;
+		float Height = 0.0f;
+		float FingerContactConfidence = 0.0f;
+		int32 DeviceContactID = 0;
+	};
+
+	/* Bitwise operators for EPenStatus to be used as flags without additional casts. */
+	
+	inline EPenStatus operator|(EPenStatus A, EPenStatus B)
+	{
+		return static_cast<EPenStatus>(static_cast<std::underlying_type_t<EPenStatus>>(A) | static_cast<std::underlying_type_t<EPenStatus>>(B));
+	}
+
+	inline EPenStatus operator&(EPenStatus A, EPenStatus B)
+	{
+		return static_cast<EPenStatus>(static_cast<std::underlying_type_t<EPenStatus>>(A) & static_cast<std::underlying_type_t<EPenStatus>>(B));
+	}
+
+	inline EPenStatus operator^(EPenStatus A, EPenStatus B)
+	{
+		return static_cast<EPenStatus>(static_cast<std::underlying_type_t<EPenStatus>>(A) ^ static_cast<std::underlying_type_t<EPenStatus>>(B));
+	}
+
+	inline EPenStatus operator~(EPenStatus A)
+	{
+		return static_cast<EPenStatus>(~static_cast<std::underlying_type_t<EPenStatus>>(A));
+	}
+}
diff --git a/Engine/Plugins/Editor/StylusInput/Source/StylusInput/Public/StylusInputTabletContext.h b/Engine/Plugins/Editor/StylusInput/Source/StylusInput/Public/StylusInputTabletContext.h
new file mode 100644
--- /dev/null
+++ b/Engine/Plugins/Editor/StylusInput/Source/StylusInput/Public/StylusInputTabletContext.h
@@ -0,0 +1,206 @@
+﻿// Copyright Epic Games, Inc. All Rights Reserved.
+
+#pragma once
+
+#include <type_traits>
+#include <Containers/UnrealString.h>
+#include <HAL/Platform.h>
+
+namespace UE::StylusInput
+{
+	/**
+	 * Flags describing tablet hardware capabilities.
+	 */
+	enum class ETabletHardwareCapabilities : uint8
+	{
+		None					= 0,
+		Integrated				= 1 << 0, ///< The digitizer is integrated with the display.
+		CursorMustTouch			= 1 << 1, ///< The cursor must be in physical contact with the device to report position.
+		HardProximity			= 1 << 2, ///< The device can generate in-air packets when the cursor is in the physical detection range (proximity) of the device.
+		CursorsHavePhysicalIds	= 1 << 3  ///< The device can uniquely identify the active cursor.
+	};
+
+	/**
+	 * Flags denoting the properties supported by a tablet.
+	 */
+	enum class ETabletSupportedProperties : uint32
+	{
+		None				    = 0,
+		X					    = 1 <<  0, ///< The x-coordinate in the coordinate space of the window that was used to create the stylus input instance. Each packet contains this property by default.
+		Y					    = 1 <<  1, ///< The y-coordinate in the coordinate space of the window that was used to create the stylus input instance. Each packet contains this property by default.
+		Z					    = 1 <<  2, ///< The z-coordinate or distance of the pen tip from the tablet surface. The TabletPropertyMetricUnit enumeration  type determines the unit of measurement for this property.
+		PacketStatus		    = 1 <<  3, ///< Contains one or more of the following flag values: The cursor is touching the drawing surface (Value = 1). The cursor is inverted. For example, the eraser end of the pen is pointing toward the surface (Value = 2). Not used (Value = 4). The barrel button is pressed (Value = 8).
+		TimerTick			    = 1 <<  4, ///< The time the packet was generated.
+		SerialNumber		    = 1 <<  5, ///< The packet property for identifying the packet. This is the same value you use to retrieve the packet from the packet queue.
+		NormalPressure		    = 1 <<  6, ///< The pressure of the pen tip perpendicular to the tablet surface. The greater the pressure on the pen tip, the  more ink that is drawn.
+		TangentPressure		    = 1 <<  7, ///< The pressure of the pen tip along the plane of the tablet surface.
+		ButtonPressure		    = 1 <<  8, ///< The pressure on a pressure sensitive button.
+		XTiltOrientation	    = 1 <<  9, ///< The angle between the y,z-plane and the pen and y-axis plane. Applies to a pen cursor. The value is 0 when the pen is perpendicular to the drawing surface and is positive when the pen is to the right of perpendicular.
+		YTiltOrientation	    = 1 << 10, ///< The angle between the x,z-plane and the pen and x-axis plane. Applies to a pen cursor. The value is 0 when the pen is perpendicular to the drawing surface and is positive when the pen is upward or away from the user.
+		AzimuthOrientation	    = 1 << 11, ///< The clockwise rotation of the cursor about the z-axis through a full circular range.
+		AltitudeOrientation     = 1 << 12, ///< The angle between the axis of the pen and the surface of the tablet. The value is 0 when the pen is parallel to the surface and 90 when the pen is perpendicular to the surface. The values are negative when the pen is inverted.
+		TwistOrientation	    = 1 << 13, ///< The clockwise rotation of the cursor about its own axis.
+		PitchRotation		    = 1 << 14, ///< The packet property that indicates whether the tip is above or below a horizontal line that is perpendicular to the writing surface. Note: This requires a 3D digitizer. The value is positive if the tip is above the line and negative if it is below the line. For example, if you hold the pen in front of you and write on an imaginary wall, the pitch is positive if the tip is above a line extending from you to the wall.
+		RollRotation		    = 1 << 15, ///< The clockwise rotation of the pen around its own axis. Note: This requires a 3D digitizer.
+		YawRotation			    = 1 << 16, ///< The angle of the pen to the left or right around the center of its horizontal axis when the pen is horizontal. Note: This requires a 3D digitizer. If you hold the pen in front of you and write on an imaginary wall, zero yaw indicates that the pen is perpendicular to the wall. The value is negative if the tip is to the left of perpendicular and positive if the tip is to the right of perpendicular.
+		Width				    = 1 << 17, ///< The width of the contact area on a touch digitizer.
+		Height					= 1 << 18, ///< The height of the contact area on a touch digitizer.
+		FingerContactConfidence = 1 << 19, ///< The level of confidence that there was finger contact on a touch digitizer.
+		DeviceContactID			= 1 << 20, ///< The device contact identifier for a packet, e.g. to identify individual fingers.
+	};
+
+	/**
+	 * Identifies and describes the capabilities of a tablet drawing surface.
+	 */
+	class IStylusInputTabletContext
+	{
+	public:
+		virtual ~IStylusInputTabletContext() = default;
+
+		/**
+		 * Returns the unique identifier for the tablet context.
+		 * The returned ID is only unique and identical for a given stylus input instance, i.e. the same tablet device might have different IDs for different
+		 * stylus input instances, and different tablet devices might have the same ID in two separate stylus input instances.
+		 *
+		 * @returns The unique identifier for the tablet context for the respective stylus input instance.
+		 */
+		STYLUSINPUT_API virtual uint32 GetID() const = 0;
+
+		/**
+		 * Returns the name of the tablet context.
+		 * This usually allows to identify the tablet device based on the hardware product name and model.
+		 *
+		 * @returns The name of the tablet context.
+		 */
+		STYLUSINPUT_API virtual FString GetName() const = 0;
+
+		/**
+		 * Returns the dimensions of the digitizer surface in device coordinates.
+		 * This mainly gives an indication of the spatial resolution of the device, and it is not related to the coordinate space of the window that was used to
+		 * create the stylus input instance.
+		 *
+		 * @returns The dimensions of the digitizer surface.
+		 */
+		STYLUSINPUT_API virtual FIntRect GetInputRectangle() const = 0;
+
+		/**
+		 * Returns the hardware capabilities of the device associated with the tablet context.
+		 *
+		 * @return The hardware capability flags.
+		 */
+		STYLUSINPUT_API virtual ETabletHardwareCapabilities GetHardwareCapabilities() const = 0;
+
+		/**
+		 * Returns the supported properties for a tablet context.
+		 * Only explicitly supported properties for a tablet context provide valid values in the packet data with the respective tablet context ID.
+		 *
+		 * @returns The supported properties flags.
+		 */
+		STYLUSINPUT_API virtual ETabletSupportedProperties GetSupportedProperties() const = 0;
+	};
+
+	/**
+	 * Description of a button on a stylus pen.
+	 */
+	class IStylusInputStylusButton
+	{
+	public:
+		virtual ~IStylusInputStylusButton() = default;
+
+		/**
+		 * @returns The GUID string for the button. 
+		 */
+		STYLUSINPUT_API virtual FString GetID() const = 0;
+
+		/**
+		 * @returns  The name of the button.
+		 */
+		STYLUSINPUT_API virtual FString GetName() const = 0;
+	};
+
+	/**
+	 * Description of a stylus pen.
+	 */
+	class IStylusInputStylusInfo
+	{
+	public:
+		virtual ~IStylusInputStylusInfo() = default;
+
+		/**
+		 * Returns the identifier for the type of stylus being used.
+		 * Note that different hardware pens can have the same ID within a tablet context.
+		 *
+		 * @returns Identifier of the stylus pen.
+		 */
+		STYLUSINPUT_API virtual uint32 GetID() const = 0;
+
+		/**
+		 * @returns The name of the stylus pen.
+		 */
+		STYLUSINPUT_API virtual FString GetName() const = 0;
+
+		/**
+		 * Returns the number of buttons on the stylus.
+		 * Note that some hardware buttons might not show up in the stylus info provided by the tablet context.
+		 *
+		 * @returns The number of buttons on the stylus.
+		 */
+		STYLUSINPUT_API virtual uint32 GetNumButtons() const = 0;
+
+		/**
+		 * Returns a description of the button for a given index.
+		 *
+		 * @param Index A number in [0, @see #GetNumButtons).
+		 * @returns Pointer to a stylus button.
+		 */
+		STYLUSINPUT_API virtual const IStylusInputStylusButton* GetButton(int32 Index) const = 0;
+	};
+
+	/* Bitwise operators for enums to be used as flags without additional casts. */
+
+	inline ETabletHardwareCapabilities operator|(ETabletHardwareCapabilities A, ETabletHardwareCapabilities B)
+	{
+		return static_cast<ETabletHardwareCapabilities>(static_cast<std::underlying_type_t<ETabletHardwareCapabilities>>(A) |
+			static_cast<std::underlying_type_t<ETabletHardwareCapabilities>>(B));
+	}
+
+	inline ETabletHardwareCapabilities operator&(ETabletHardwareCapabilities A, ETabletHardwareCapabilities B)
+	{
+		return static_cast<ETabletHardwareCapabilities>(static_cast<std::underlying_type_t<ETabletHardwareCapabilities>>(A) &
+			static_cast<std::underlying_type_t<ETabletHardwareCapabilities>>(B));
+	}
+
+	inline ETabletHardwareCapabilities operator^(ETabletHardwareCapabilities A, ETabletHardwareCapabilities B)
+	{
+		return static_cast<ETabletHardwareCapabilities>(static_cast<std::underlying_type_t<ETabletHardwareCapabilities>>(A) ^
+			static_cast<std::underlying_type_t<ETabletHardwareCapabilities>>(B));
+	}
+
+	inline ETabletHardwareCapabilities operator~(ETabletHardwareCapabilities A)
+	{
+		return static_cast<ETabletHardwareCapabilities>(~static_cast<std::underlying_type_t<ETabletHardwareCapabilities>>(A));
+	}
+
+	inline ETabletSupportedProperties operator|(ETabletSupportedProperties A, ETabletSupportedProperties B)
+	{
+		return static_cast<ETabletSupportedProperties>(static_cast<std::underlying_type_t<ETabletSupportedProperties>>(A) |
+			static_cast<std::underlying_type_t<ETabletSupportedProperties>>(B));
+	}
+
+	inline ETabletSupportedProperties operator&(ETabletSupportedProperties A, ETabletSupportedProperties B)
+	{
+		return static_cast<ETabletSupportedProperties>(static_cast<std::underlying_type_t<ETabletSupportedProperties>>(A) &
+			static_cast<std::underlying_type_t<ETabletSupportedProperties>>(B));
+	}
+
+	inline ETabletSupportedProperties operator^(ETabletSupportedProperties A, ETabletSupportedProperties B)
+	{
+		return static_cast<ETabletSupportedProperties>(static_cast<std::underlying_type_t<ETabletSupportedProperties>>(A) ^
+			static_cast<std::underlying_type_t<ETabletSupportedProperties>>(B));
+	}
+
+	inline ETabletSupportedProperties operator~(ETabletSupportedProperties A)
+	{
+		return static_cast<ETabletSupportedProperties>(~static_cast<std::underlying_type_t<ETabletSupportedProperties>>(A));
+	}
+}
diff --git a/Engine/Plugins/Editor/StylusInput/Source/StylusInput/StylusInput.Build.cs b/Engine/Plugins/Editor/StylusInput/Source/StylusInput/StylusInput.Build.cs
--- a/Engine/Plugins/Editor/StylusInput/Source/StylusInput/StylusInput.Build.cs
+++ b/Engine/Plugins/Editor/StylusInput/Source/StylusInput/StylusInput.Build.cs
@@ -9,26 +9,21 @@ namespace UnrealBuildTool.Rules
 			PublicDependencyModuleNames.AddRange(
 				new string[]
 				{
+					"Core",
 					"CoreUObject",
 					"EditorSubsystem",
 					"Engine",
-					"EditorFramework",
-					"UnrealEd"
-					// ... add other public dependencies that you statically link with here ...
+					"Slate",
+					"UnrealEd",
 				}
-				);
+			);
 
 			PrivateDependencyModuleNames.AddRange(
 				new string[]
 				{
-					"Core",
-					"MainFrame",
 					"SlateCore",
-					"Slate",
-					"WorkspaceMenuStructure"
-					// ... add private dependencies that you statically link with here ...
 				}
-				);
+			);
 		}
 	}
 }
diff --git a/Engine/Plugins/Editor/StylusInput/Source/StylusInputDebugWidget/Private/StylusInputDebugPaintWidget.cpp b/Engine/Plugins/Editor/StylusInput/Source/StylusInputDebugWidget/Private/StylusInputDebugPaintWidget.cpp
new file mode 100644
--- /dev/null
+++ b/Engine/Plugins/Editor/StylusInput/Source/StylusInputDebugWidget/Private/StylusInputDebugPaintWidget.cpp
@@ -0,0 +1,118 @@
+﻿// Copyright Epic Games, Inc. All Rights Reserved.
+
+#include "StylusInputDebugPaintWidget.h"
+
+#include <StylusInputPacket.h>
+#include <Brushes/SlateColorBrush.h>
+#include <Framework/Application/SlateApplication.h>
+#include <Styling/AppStyle.h>
+
+namespace UE::StylusInput::DebugWidget
+{
+	int32 DrawBackground(const FGeometry& AllottedGeometry, FSlateWindowElementList& OutDrawElements, int32 LayerId)
+	{
+		const FLinearColor BackgroundColor(FVector3f(0.05f));
+		FSlateDrawElement::MakeBox(OutDrawElements, LayerId, AllottedGeometry.ToPaintGeometry(), FAppStyle::GetBrush(TEXT("WhiteBrush")),
+		                           ESlateDrawEffect::None, BackgroundColor);
+		return ++LayerId;
+	}
+
+	void DrawCircle(FVector2f Position, const float Size, const FLinearColor& Color, const FGeometry& AllottedGeometry,
+	                FSlateWindowElementList& OutDrawElements, int32 LayerId)
+	{
+		const FPaintGeometry PaintGeometry = AllottedGeometry.ToPaintGeometry(Slate::FDeprecateVector2DParameter(Size, Size), FSlateLayoutTransform(Position));
+
+		FSlateColorBrush ColorBrush = FLinearColor::White;
+		ColorBrush.DrawAs = ESlateBrushDrawType::Type::RoundedBox;
+
+		FSlateDrawElement::MakeBox(OutDrawElements, LayerId, PaintGeometry, &ColorBrush, ESlateDrawEffect::None, Color);
+	}
+
+	void SStylusInputDebugPaintWidget::Construct(const FArguments& Args)
+	{
+		PaintPackets.Reserve(1024);
+	}
+
+	void SStylusInputDebugPaintWidget::Add(const FStylusInputPacket& Packet)
+	{
+		PaintPackets.Emplace(FPaintPacket{
+			FVector2f{Packet.X, Packet.Y}, Packet.NormalPressure, FPlatformTime::ToMilliseconds(FPlatformTime::Cycles()), Packet.PenStatus
+		});
+	}
+
+	int32 SStylusInputDebugPaintWidget::OnPaint(const FPaintArgs& Args, const FGeometry& AllottedGeometry, const FSlateRect& MyCullingRect,
+	                                            FSlateWindowElementList& OutDrawElements, int32 LayerId, const FWidgetStyle& InWidgetStyle,
+	                                            bool bParentEnabled) const
+	{
+		LayerId = DrawBackground(AllottedGeometry, OutDrawElements, LayerId);
+		LayerId = DrawPaintPackets(AllottedGeometry, OutDrawElements, LayerId);
+
+		return LayerId;
+	}
+
+	int32 SStylusInputDebugPaintWidget::DrawPaintPackets(const FGeometry& AllottedGeometry, FSlateWindowElementList& OutDrawElements, int32 LayerId) const
+	{
+		if (PaintPackets.IsEmpty())
+		{
+			return LayerId;
+		}
+
+		constexpr float MinSize = 2.0f;
+		constexpr float MaxSize = 15.0f;
+		constexpr float FadeOutTimeMSPenDown = 3000.0f;
+		constexpr float FadeOutTimeMSPenUp = 1000.0f;
+
+		constexpr FLinearColor ColorPenDown = {0.0f, 0.2f, 1.0f};
+		constexpr FLinearColor ColorPenDownInverted = {1.0f, 0.1f, 0.0f};
+		constexpr FLinearColor ColorPenUp = {0.5f, 0.5f, 0.5f};
+
+		Slate::FDeprecateVector2DResult WindowPosition = FVector2f::ZeroVector;
+		float DPIScaleFactor = 1.0f;
+
+		const TSharedPtr<SWindow> Window = FSlateApplication::Get().FindWidgetWindow(AsShared());
+		if (ensure(Window))
+		{
+			WindowPosition = Window->GetPaintSpaceGeometry().GetAbsolutePosition();
+			DPIScaleFactor = Window ? Window->GetDPIScaleFactor() : 1.0f;
+		}
+
+		const Slate::FDeprecateVector2DResult PaintWidgetPosition = GetPaintSpaceGeometry().GetAbsolutePosition();
+		const FDeprecateSlateVector2D PositionOffset = PaintWidgetPosition - WindowPosition;
+
+		const float CurrentTimeMS = FPlatformTime::ToMilliseconds(FPlatformTime::Cycles());
+
+		int32 NumToRemove = 0;
+		for (const FPaintPacket& PaintPacket : PaintPackets)
+		{
+			const bool bPenIsDown = (PaintPacket.PenStatus & EPenStatus::CursorIsTouching) != EPenStatus::None;
+			const bool bPenIsInverted = (PaintPacket.PenStatus & EPenStatus::CursorIsInverted) != EPenStatus::None;
+
+			const float DeltaTimeMS = CurrentTimeMS - PaintPacket.TimeAddedMS;
+			const float FadeOutTimeMS = bPenIsDown ? FadeOutTimeMSPenDown : FadeOutTimeMSPenUp;
+
+			if (DeltaTimeMS >= FadeOutTimeMS)
+			{
+				++NumToRemove;
+			}
+			else
+			{
+				FLinearColor Color = bPenIsDown ? (bPenIsInverted ? ColorPenDownInverted : ColorPenDown) : ColorPenUp;
+				Color.A = (FadeOutTimeMS - DeltaTimeMS) / FadeOutTimeMS;
+
+				const float Size = FMath::Max(MinSize, PaintPacket.NormalPressure * MaxSize);
+
+				const FVector2f TransformedPosition = (PaintPacket.Position - PositionOffset - Size / 2.0f) / DPIScaleFactor;
+
+				DrawCircle(TransformedPosition, Size, Color, AllottedGeometry, OutDrawElements, LayerId);
+			}
+		}
+
+		if (NumToRemove)
+		{
+			checkSlow(NumToRemove <= PaintPackets.Num());
+			PaintPackets.PopFrontNoCheck(NumToRemove);
+		}
+
+		return ++LayerId;
+	}
+}
diff --git a/Engine/Plugins/Editor/StylusInput/Source/StylusInputDebugWidget/Private/StylusInputDebugPaintWidget.h b/Engine/Plugins/Editor/StylusInput/Source/StylusInputDebugWidget/Private/StylusInputDebugPaintWidget.h
new file mode 100644
--- /dev/null
+++ b/Engine/Plugins/Editor/StylusInput/Source/StylusInputDebugWidget/Private/StylusInputDebugPaintWidget.h
@@ -0,0 +1,45 @@
+﻿// Copyright Epic Games, Inc. All Rights Reserved.
+
+#pragma once
+
+#include <StylusInputPacket.h>
+#include <Widgets/SCompoundWidget.h>
+#include <Containers/RingBuffer.h>
+
+namespace UE::StylusInput
+{
+	struct FStylusInputPacket;
+}
+
+namespace UE::StylusInput::DebugWidget
+{
+	class SStylusInputDebugPaintWidget : public SCompoundWidget
+	{
+	public:
+		SLATE_BEGIN_ARGS(SStylusInputDebugPaintWidget)
+			{
+			}
+
+		SLATE_END_ARGS()
+
+		void Construct(const FArguments& Args);
+
+		void Add(const FStylusInputPacket& Packet);
+
+		virtual int32 OnPaint(const FPaintArgs& Args, const FGeometry& AllottedGeometry, const FSlateRect& MyCullingRect,
+							  FSlateWindowElementList& OutDrawElements, int32 LayerId, const FWidgetStyle& InWidgetStyle, bool bParentEnabled) const override;
+
+	private:
+		struct FPaintPacket
+		{
+			FVector2f Position = FVector2f::ZeroVector;
+			float NormalPressure = 0.0f;
+			float TimeAddedMS = 0.0f;
+			EPenStatus PenStatus = EPenStatus::None;
+		};
+
+		int32 DrawPaintPackets(const FGeometry& AllottedGeometry, FSlateWindowElementList& OutDrawElements, int32 LayerId) const;
+
+		mutable TRingBuffer<FPaintPacket> PaintPackets;
+	};
+}
diff --git a/Engine/Plugins/Editor/StylusInput/Source/StylusInputDebugWidget/Private/StylusInputDebugWidget.cpp b/Engine/Plugins/Editor/StylusInput/Source/StylusInputDebugWidget/Private/StylusInputDebugWidget.cpp
new file mode 100644
--- /dev/null
+++ b/Engine/Plugins/Editor/StylusInput/Source/StylusInputDebugWidget/Private/StylusInputDebugWidget.cpp
@@ -0,0 +1,949 @@
+﻿// Copyright Epic Games, Inc. All Rights Reserved.
+
+#include "StylusInputDebugWidget.h"
+
+#include <StylusInput.h>
+#include <StylusInputTabletContext.h>
+#include <Framework/Application/SlateApplication.h>
+#include <Framework/MultiBox/MultiBoxBuilder.h>
+#include <Widgets/SBoxPanel.h>
+#include <Widgets/SOverlay.h>
+#include <Widgets/Input/SComboButton.h>
+#include <Widgets/Layout/SSplitter.h>
+#include <Widgets/Text/SMultiLineEditableText.h>
+#include <Widgets/Text/STextBlock.h>
+
+#include "StylusInputDebugPaintWidget.h"
+
+#define LOCTEXT_NAMESPACE "StylusInputDebugWidget"
+
+namespace UE::StylusInput::DebugWidget
+{
+	DECLARE_LOG_CATEGORY_EXTERN(LogStylusInputDebugWidget, Log, All)
+
+	DEFINE_LOG_CATEGORY(LogStylusInputDebugWidget);
+
+	inline void LogError(const FString& Message)
+	{
+		UE_LOG(LogStylusInputDebugWidget, Error, TEXT("%s"), *Message);
+	}
+
+	inline void LogWarning(const FString& Message)
+	{
+		UE_LOG(LogStylusInputDebugWidget, Warning, TEXT("%s"), *Message);
+	}
+
+	inline void LogVerbose(const FString& Message)
+	{
+		UE_LOG(LogStylusInputDebugWidget, Verbose, TEXT("%s"), *Message);
+	}
+
+	void SStylusInputDebugWidget::Construct(const FArguments&)
+	{
+		// One-time timer to initialize stylus plugin as soon as widget is up.
+		RegisterActiveTimer(
+			0.0f,
+			FWidgetActiveTimerDelegate::CreateLambda([this](double, float)
+			{
+				AcquireStylusInput();
+				RegisterEventHandler();
+				return EActiveTimerReturnType::Stop;
+			}));
+
+		ChildSlot
+			.VAlign(VAlign_Fill)
+			.HAlign(HAlign_Fill)
+			[
+				SNew(SSplitter)
+				.PhysicalSplitterHandleSize(2.0f)
+				+ SSplitter::Slot()
+				[
+					SNew(SOverlay)
+					+ SOverlay::Slot()
+					[
+						SAssignNew(PaintWidget, SStylusInputDebugPaintWidget)
+					]
+					+ SOverlay::Slot()
+					.Padding(2.0f)
+					[
+						SNew(SVerticalBox)
+						+ SVerticalBox::Slot()
+						.AutoHeight()
+						[
+							SNew(SHorizontalBox)
+							+ SHorizontalBox::Slot()
+							.AutoWidth()
+							[
+								SNew(SVerticalBox)
+								+ SVerticalBox::Slot()
+								.AutoHeight()
+								[
+									SNew(STextBlock)
+									.Text_Lambda([&TabletContext = LastPacketData.TabletContext]
+									{
+										return TabletContext
+											       ? FText::Format(
+												       LOCTEXT("TabletContextID", "Tablet Context ID: {0}"),
+												       FText::FromString(FString::Printf(TEXT("%x"), TabletContext->GetID())))
+											       : FText();
+									})
+								]
+								+ SVerticalBox::Slot()
+								.AutoHeight()
+								[
+									SNew(STextBlock)
+									.Text_Lambda([&TabletContext = LastPacketData.TabletContext]
+									{
+										return TabletContext
+											       ? FText::Format(
+												       LOCTEXT("TabletContextName", "Name: {0}"),
+												       FText::FromString(TabletContext->GetName()))
+											       : FText();
+									})
+								]
+								+ SVerticalBox::Slot()
+								.AutoHeight()
+								[
+									SNew(STextBlock)
+									.Text_Lambda([&TabletContext = LastPacketData.TabletContext]
+									{
+										if (TabletContext)
+										{
+											const FIntRect InputRectangle = TabletContext->GetInputRectangle();
+											return FText::Format(
+												LOCTEXT("TabletContextInputRectangle", "Input Rectangle: ({0}, {1}) x ({2}, {3})"),
+												InputRectangle.Min.X, InputRectangle.Min.Y, InputRectangle.Max.X, InputRectangle.Max.Y);
+										}
+										return FText();
+									})
+								]
+								+ SVerticalBox::Slot()
+								.AutoHeight()
+								[
+									SNew(STextBlock)
+									.Text_Lambda([&TabletContext = LastPacketData.TabletContext]
+									{
+										if (TabletContext)
+										{
+											const ETabletHardwareCapabilities Capabilities = TabletContext->GetHardwareCapabilities();
+
+											FString CapabilitiesStr;
+											if (Capabilities == ETabletHardwareCapabilities::None)
+											{
+												CapabilitiesStr = "None";
+											}
+											else
+											{
+												if ((Capabilities & ETabletHardwareCapabilities::Integrated) != ETabletHardwareCapabilities::None)
+												{
+													CapabilitiesStr += "Integrated";
+												}
+												if ((Capabilities & ETabletHardwareCapabilities::CursorMustTouch) != ETabletHardwareCapabilities::None)
+												{
+													CapabilitiesStr += CapabilitiesStr.IsEmpty() ? "CursorMustTouch" : " & CursorMustTouch";
+												}
+												if ((Capabilities & ETabletHardwareCapabilities::HardProximity) != ETabletHardwareCapabilities::None)
+												{
+													CapabilitiesStr += CapabilitiesStr.IsEmpty() ? "HardProximity" : " & HardProximity";
+												}
+												if ((Capabilities & ETabletHardwareCapabilities::CursorsHavePhysicalIds) != ETabletHardwareCapabilities::None)
+												{
+													CapabilitiesStr += CapabilitiesStr.IsEmpty() ? "CursorsHavePhysicalIds" : " & CursorsHavePhysicalIds";
+												}
+											}
+
+											if (CapabilitiesStr.IsEmpty())
+											{
+												CapabilitiesStr = FString::Format(TEXT("Unknown ({0})"), {
+													                                  static_cast<std::underlying_type_t<ETabletHardwareCapabilities>>(
+														                                  Capabilities)
+												                                  });
+											}
+
+											return FText::Format(
+												LOCTEXT("TabletContextHardwareCapabilities", "Hardware Capabilities: {0}"), FText::FromString(CapabilitiesStr));
+										}
+										return FText();
+									})
+								]
+								+ SVerticalBox::Slot()
+								.AutoHeight()
+								[
+									SNew(STextBlock)
+									.Text_Lambda([&StylusInfo = LastPacketData.StylusInfo]
+									{
+										return StylusInfo
+											       ? FText::Format(
+												       LOCTEXT("StylusID", "Stylus ID: {0}"),
+												       FText::FromString(FString::Printf(TEXT("%d"), StylusInfo->GetID())))
+											       : FText();
+									})
+								]
+								+ SVerticalBox::Slot()
+								.AutoHeight()
+								[
+									SNew(STextBlock)
+									.Text_Lambda([&StylusInfo = LastPacketData.StylusInfo]
+									{
+										return StylusInfo
+											       ? FText::Format(
+												       LOCTEXT("StylusName", "Stylus Name: {0}"),
+												       FText::FromString(StylusInfo->GetName()))
+											       : FText();
+									})
+								]
+								+ SVerticalBox::Slot()
+								.AutoHeight()
+								[
+									SNew(STextBlock)
+									.Text_Lambda([&StylusInfo = LastPacketData.StylusInfo]
+									{
+										if (StylusInfo)
+										{
+											FString ButtonsString = "<none>";
+											for (int32 Index = 0, Num = StylusInfo->GetNumButtons(); Index < Num; ++Index)
+											{
+												if (const auto* Button = StylusInfo->GetButton(Index))
+												{
+													if (Index == 0)
+													{
+														ButtonsString = Button->GetName();
+													}
+													else
+													{
+														ButtonsString += ", " + Button->GetName();
+													}
+												}
+											}
+
+											return FText::Format(LOCTEXT("StylusButtons", "Stylus Buttons: {0}"), FText::FromString(ButtonsString));
+										}
+										return FText();
+									})
+								]
+							]
+							+ SHorizontalBox::Slot()
+							.FillWidth(1.0f)
+							[
+								SNullWidget::NullWidget
+							]
+							+ SHorizontalBox::Slot()
+							.AutoWidth()
+							[
+								SNew(SVerticalBox)
+								+ SVerticalBox::Slot()
+								.AutoHeight()
+								[
+									SNew(SComboButton)
+									.OnGetMenuContent(this, &SStylusInputDebugWidget::GetEventHandlerThreadMenu)
+									.ButtonContent()
+									[
+										SNew(STextBlock)
+										.Text_Lambda([&EventHandlerThread = EventHandlerThread]
+										{
+											return FText::FromString(
+												EventHandlerThread == EEventHandlerThread::Asynchronous ? "Asynchronous" : "On Game Thread");
+										})
+									]
+								]
+								+ SVerticalBox::Slot()
+								.FillHeight(1.0f)
+								[
+									SNullWidget::NullWidget
+								]
+							]
+						]
+						+ SVerticalBox::Slot()
+						.FillHeight(1.0f)
+						[
+							SNullWidget::NullWidget
+						]
+						+ SVerticalBox::Slot()
+						.AutoHeight()
+						[
+							SNew(STextBlock)
+							.Text_Lambda([&TabletContext = LastPacketData.TabletContext, &StylusInput = StylusInput, &EventHandlerThread = EventHandlerThread]
+							{
+								if (TabletContext && StylusInput)
+								{
+									return FText::Format(
+										LOCTEXT("PacketsPerSecond", "Packets Per Second: {0}"),
+										FMath::RoundToInt32(StylusInput->GetPacketsPerSecond(EventHandlerThread)));
+								}
+								return FText();
+							})
+						]
+						/*+ SVerticalBox::Slot()
+						.AutoHeight()
+						[
+							SNew(STextBlock)
+							.Text_Lambda([&bIsSet = LastPacketData.bIsSet, &CursorID = LastPacketData.Packet.CursorID]
+							{
+								return bIsSet ? FText::Format(LOCTEXT("CursorID", "Cursor ID: {0}"), CursorID) : FText();
+							})
+						]*/
+						+ SVerticalBox::Slot()
+						.AutoHeight()
+						[
+							SNew(STextBlock)
+							.Visibility_Lambda([&bIsSet = LastPacketData.bIsSet, &TabletContext = LastPacketData.TabletContext]
+							{
+								return bIsSet && TabletContext && (TabletContext->GetSupportedProperties() & ETabletSupportedProperties::TimerTick)
+								       != ETabletSupportedProperties::None
+									       ? EVisibility::Visible
+									       : EVisibility::Collapsed;
+							})
+							.Text_Lambda([&TimerTick = LastPacketData.Packet.TimerTick]
+							{
+								return FText::Format(LOCTEXT("TimerTick", "Timer Tick: {0}"), TimerTick);
+							})
+						]
+						+ SVerticalBox::Slot()
+						.AutoHeight()
+						[
+							SNew(STextBlock)
+							.Visibility_Lambda([&bIsSet = LastPacketData.bIsSet, &TabletContext = LastPacketData.TabletContext]
+							{
+								return bIsSet && TabletContext && (TabletContext->GetSupportedProperties() & ETabletSupportedProperties::SerialNumber)
+								       != ETabletSupportedProperties::None
+									       ? EVisibility::Visible
+									       : EVisibility::Collapsed;
+							})
+							.Text_Lambda([&SerialNumber = LastPacketData.Packet.SerialNumber]
+							{
+								return FText::Format(LOCTEXT("SerialNumber", "Serial Number: {0}"), SerialNumber);
+							})
+						]
+						+ SVerticalBox::Slot()
+						.AutoHeight()
+						[
+							SNew(STextBlock)
+							.Text_Lambda([&bIsSet = LastPacketData.bIsSet, &PenStatus = LastPacketData.Packet.PenStatus]
+							{
+								if (bIsSet)
+								{
+									FString PenStatusStr;
+									if (PenStatus == EPenStatus::None)
+									{
+										PenStatusStr = "None";
+									}
+									else
+									{
+										if ((PenStatus & EPenStatus::CursorIsTouching) != EPenStatus::None)
+										{
+											PenStatusStr += "CursorIsTouching";
+										}
+										if ((PenStatus & EPenStatus::CursorIsInverted) != EPenStatus::None)
+										{
+											PenStatusStr += PenStatusStr.IsEmpty() ? "CursorIsInverted" : " & CursorIsInverted";
+										}
+										if ((PenStatus & EPenStatus::NotUsed) != EPenStatus::None)
+										{
+											PenStatusStr += PenStatusStr.IsEmpty() ? "NotUsed" : " & NotUsed";
+										}
+										if ((PenStatus & EPenStatus::BarrelButtonPressed) != EPenStatus::None)
+										{
+											PenStatusStr += PenStatusStr.IsEmpty() ? "BarrelButtonPressed" : " & BarrelButtonPressed";
+										}
+									}
+
+									if (PenStatusStr.IsEmpty())
+									{
+										PenStatusStr = FString::Format(TEXT("Unknown ({0})"), {
+											                               static_cast<std::underlying_type_t<EPenStatus>>(PenStatus)
+										                               });
+									}
+
+									return FText::Format(LOCTEXT("PenStatus", "Pen Status: {0}"), FText::FromString(PenStatusStr));
+								}
+								return FText();
+							})
+						]
+						+ SVerticalBox::Slot()
+						.AutoHeight()
+						[
+							SNew(STextBlock)
+							.Visibility_Lambda([&bIsSet = LastPacketData.bIsSet, &TabletContext = LastPacketData.TabletContext]
+							{
+								return bIsSet && TabletContext && (TabletContext->GetSupportedProperties() & ETabletSupportedProperties::X)
+								       != ETabletSupportedProperties::None
+									       ? EVisibility::Visible
+									       : EVisibility::Collapsed;
+							})
+							.Text_Lambda([&X = LastPacketData.Packet.X]
+							{
+								return FText::Format(LOCTEXT("X", "X: {0}"), X);
+							})
+						]
+						+ SVerticalBox::Slot()
+						.AutoHeight()
+						[
+							SNew(STextBlock)
+							.Visibility_Lambda([&bIsSet = LastPacketData.bIsSet, &TabletContext = LastPacketData.TabletContext]
+							{
+								return bIsSet && TabletContext && (TabletContext->GetSupportedProperties() & ETabletSupportedProperties::Y)
+								       != ETabletSupportedProperties::None
+									       ? EVisibility::Visible
+									       : EVisibility::Collapsed;
+							})
+							.Text_Lambda([&Y = LastPacketData.Packet.Y]
+							{
+								return FText::Format(LOCTEXT("Y", "Y: {0}"), Y);
+							})
+						]
+						+ SVerticalBox::Slot()
+						.AutoHeight()
+						[
+							SNew(STextBlock)
+							.Visibility_Lambda([&bIsSet = LastPacketData.bIsSet, &TabletContext = LastPacketData.TabletContext]
+							{
+								return bIsSet && TabletContext && (TabletContext->GetSupportedProperties() & ETabletSupportedProperties::Z)
+								       != ETabletSupportedProperties::None
+									       ? EVisibility::Visible
+									       : EVisibility::Collapsed;
+							})
+							.Text_Lambda([&Z = LastPacketData.Packet.Z]
+							{
+								return FText::Format(LOCTEXT("Z", "Z: {0}"), Z);
+							})
+						]
+						+ SVerticalBox::Slot()
+						.AutoHeight()
+						[
+							SNew(STextBlock)
+							.Visibility_Lambda([&bIsSet = LastPacketData.bIsSet, &TabletContext = LastPacketData.TabletContext]
+							{
+								return bIsSet && TabletContext && (TabletContext->GetSupportedProperties() & ETabletSupportedProperties::NormalPressure)
+								       != ETabletSupportedProperties::None
+									       ? EVisibility::Visible
+									       : EVisibility::Collapsed;
+							})
+							.Text_Lambda([&NormalPressure = LastPacketData.Packet.NormalPressure]
+							{
+								return FText::Format(LOCTEXT("NormalPressure", "Normal Pressure: {0}"), NormalPressure);
+							})
+						]
+						+ SVerticalBox::Slot()
+						.AutoHeight()
+						[
+							SNew(STextBlock)
+							.Visibility_Lambda([&bIsSet = LastPacketData.bIsSet, &TabletContext = LastPacketData.TabletContext]
+							{
+								return bIsSet && TabletContext && (TabletContext->GetSupportedProperties() & ETabletSupportedProperties::TangentPressure)
+								       != ETabletSupportedProperties::None
+									       ? EVisibility::Visible
+									       : EVisibility::Collapsed;
+							})
+							.Text_Lambda([&TangentPressure = LastPacketData.Packet.TangentPressure]
+							{
+								return FText::Format(LOCTEXT("TangentPressure", "Tangent Pressure: {0}"), TangentPressure);
+							})
+						]
+						+ SVerticalBox::Slot()
+						.AutoHeight()
+						[
+							SNew(STextBlock)
+							.Visibility_Lambda([&bIsSet = LastPacketData.bIsSet, &TabletContext = LastPacketData.TabletContext]
+							{
+								return bIsSet && TabletContext && (TabletContext->GetSupportedProperties() & ETabletSupportedProperties::ButtonPressure)
+								       != ETabletSupportedProperties::None
+									       ? EVisibility::Visible
+									       : EVisibility::Collapsed;
+							})
+							.Text_Lambda([&ButtonPressure = LastPacketData.Packet.ButtonPressure]
+							{
+								return FText::Format(LOCTEXT("ButtonPressure", "Button Pressure: {0}"), ButtonPressure);
+							})
+						]
+						+ SVerticalBox::Slot()
+						.AutoHeight()
+						[
+							SNew(STextBlock)
+							.Visibility_Lambda([&bIsSet = LastPacketData.bIsSet, &TabletContext = LastPacketData.TabletContext]
+							{
+								return bIsSet && TabletContext && (TabletContext->GetSupportedProperties() & ETabletSupportedProperties::XTiltOrientation)
+								       != ETabletSupportedProperties::None
+									       ? EVisibility::Visible
+									       : EVisibility::Collapsed;
+							})
+							.Text_Lambda([&XTiltOrientation = LastPacketData.Packet.XTiltOrientation]
+							{
+								return FText::Format(LOCTEXT("XTiltOrientation", "X Tilt Orientation: {0}"), XTiltOrientation);
+							})
+						]
+						+ SVerticalBox::Slot()
+						.AutoHeight()
+						[
+							SNew(STextBlock)
+							.Visibility_Lambda([&bIsSet = LastPacketData.bIsSet, &TabletContext = LastPacketData.TabletContext]
+							{
+								return bIsSet && TabletContext && (TabletContext->GetSupportedProperties() & ETabletSupportedProperties::YTiltOrientation)
+								       != ETabletSupportedProperties::None
+									       ? EVisibility::Visible
+									       : EVisibility::Collapsed;
+							})
+							.Text_Lambda([&YTiltOrientation = LastPacketData.Packet.YTiltOrientation]
+							{
+								return FText::Format(LOCTEXT("YTiltOrientation", "Y Tilt Orientation: {0}"), YTiltOrientation);
+							})
+						]
+						+ SVerticalBox::Slot()
+						.AutoHeight()
+						[
+							SNew(STextBlock)
+							.Visibility_Lambda([&bIsSet = LastPacketData.bIsSet, &TabletContext = LastPacketData.TabletContext]
+							{
+								return bIsSet && TabletContext && (TabletContext->GetSupportedProperties() & ETabletSupportedProperties::AzimuthOrientation)
+								       != ETabletSupportedProperties::None
+									       ? EVisibility::Visible
+									       : EVisibility::Collapsed;
+							})
+							.Text_Lambda([&AzimuthOrientation = LastPacketData.Packet.AzimuthOrientation]
+							{
+								return FText::Format(LOCTEXT("AzimuthOrientation", "Azimuth Orientation: {0}"), AzimuthOrientation);
+							})
+						]
+						+ SVerticalBox::Slot()
+						.AutoHeight()
+						[
+							SNew(STextBlock)
+							.Visibility_Lambda([&bIsSet = LastPacketData.bIsSet, &TabletContext = LastPacketData.TabletContext]
+							{
+								return bIsSet && TabletContext && (TabletContext->GetSupportedProperties() & ETabletSupportedProperties::AltitudeOrientation)
+								       != ETabletSupportedProperties::None
+									       ? EVisibility::Visible
+									       : EVisibility::Collapsed;
+							})
+							.Text_Lambda([&AltitudeOrientation = LastPacketData.Packet.AltitudeOrientation]
+							{
+								return FText::Format(LOCTEXT("AltitudeOrientation", "Altitude Orientation: {0}"), AltitudeOrientation);
+							})
+						]
+						+ SVerticalBox::Slot()
+						.AutoHeight()
+						[
+							SNew(STextBlock)
+							.Visibility_Lambda([&bIsSet = LastPacketData.bIsSet, &TabletContext = LastPacketData.TabletContext]
+							{
+								return bIsSet && TabletContext && (TabletContext->GetSupportedProperties() & ETabletSupportedProperties::TwistOrientation)
+								       != ETabletSupportedProperties::None
+									       ? EVisibility::Visible
+									       : EVisibility::Collapsed;
+							})
+							.Text_Lambda([&TwistOrientation = LastPacketData.Packet.TwistOrientation]
+							{
+								return FText::Format(LOCTEXT("TwistOrientation", "Twist Orientation: {0}"), TwistOrientation);
+							})
+						]
+						+ SVerticalBox::Slot()
+						.AutoHeight()
+						[
+							SNew(STextBlock)
+							.Visibility_Lambda([&bIsSet = LastPacketData.bIsSet, &TabletContext = LastPacketData.TabletContext]
+							{
+								return bIsSet && TabletContext && (TabletContext->GetSupportedProperties() & ETabletSupportedProperties::PitchRotation)
+								       != ETabletSupportedProperties::None
+									       ? EVisibility::Visible
+									       : EVisibility::Collapsed;
+							})
+							.Text_Lambda([&PitchRotation = LastPacketData.Packet.PitchRotation]
+							{
+								return FText::Format(LOCTEXT("PitchRotation", "Pitch Rotation: {0}"), PitchRotation);
+							})
+						]
+						+ SVerticalBox::Slot()
+						.AutoHeight()
+						[
+							SNew(STextBlock)
+							.Visibility_Lambda([&bIsSet = LastPacketData.bIsSet, &TabletContext = LastPacketData.TabletContext]
+							{
+								return bIsSet && TabletContext && (TabletContext->GetSupportedProperties() & ETabletSupportedProperties::RollRotation)
+								       != ETabletSupportedProperties::None
+									       ? EVisibility::Visible
+									       : EVisibility::Collapsed;
+							})
+							.Text_Lambda([&RollRotation = LastPacketData.Packet.RollRotation]
+							{
+								return FText::Format(LOCTEXT("RollRotation", "Roll Rotation: {0}"), RollRotation);
+							})
+						]
+						+ SVerticalBox::Slot()
+						.AutoHeight()
+						[
+							SNew(STextBlock)
+							.Visibility_Lambda([&bIsSet = LastPacketData.bIsSet, &TabletContext = LastPacketData.TabletContext]
+							{
+								return bIsSet && TabletContext && (TabletContext->GetSupportedProperties() & ETabletSupportedProperties::YawRotation)
+								       != ETabletSupportedProperties::None
+									       ? EVisibility::Visible
+									       : EVisibility::Collapsed;
+							})
+							.Text_Lambda([&YawRotation = LastPacketData.Packet.YawRotation]
+							{
+								return FText::Format(LOCTEXT("YawRotation", "Yaw Rotation: {0}"), YawRotation);
+							})
+						]
+						+ SVerticalBox::Slot()
+						.AutoHeight()
+						[
+							SNew(STextBlock)
+							.Visibility_Lambda([&bIsSet = LastPacketData.bIsSet, &TabletContext = LastPacketData.TabletContext]
+							{
+								return bIsSet && TabletContext && (TabletContext->GetSupportedProperties() & ETabletSupportedProperties::Width)
+								       != ETabletSupportedProperties::None
+									       ? EVisibility::Visible
+									       : EVisibility::Collapsed;
+							})
+							.Text_Lambda([&Width = LastPacketData.Packet.Width]
+							{
+								return FText::Format(LOCTEXT("Width", "Width: {0}"), Width);
+							})
+						]
+						+ SVerticalBox::Slot()
+						.AutoHeight()
+						[
+							SNew(STextBlock)
+							.Visibility_Lambda([&bIsSet = LastPacketData.bIsSet, &TabletContext = LastPacketData.TabletContext]
+							{
+								return bIsSet && TabletContext && (TabletContext->GetSupportedProperties() & ETabletSupportedProperties::Height)
+								       != ETabletSupportedProperties::None
+									       ? EVisibility::Visible
+									       : EVisibility::Collapsed;
+							})
+							.Text_Lambda([&Height = LastPacketData.Packet.Height]
+							{
+								return FText::Format(LOCTEXT("Height", "Height: {0}"), Height);
+							})
+						]
+						+ SVerticalBox::Slot()
+						.AutoHeight()
+						[
+							SNew(STextBlock)
+							.Visibility_Lambda([&bIsSet = LastPacketData.bIsSet, &TabletContext = LastPacketData.TabletContext]
+							{
+								return bIsSet && TabletContext && (TabletContext->GetSupportedProperties() &
+									       ETabletSupportedProperties::FingerContactConfidence)
+								       != ETabletSupportedProperties::None
+									       ? EVisibility::Visible
+									       : EVisibility::Collapsed;
+							})
+							.Text_Lambda([&FingerContactConfidence = LastPacketData.Packet.FingerContactConfidence]
+							{
+								return FText::Format(LOCTEXT("FingerContactConfidence", "Finger Contact Confidence: {0}"), FingerContactConfidence);
+							})
+						]
+						+ SVerticalBox::Slot()
+						.AutoHeight()
+						[
+							SNew(STextBlock)
+							.Visibility_Lambda([&bIsSet = LastPacketData.bIsSet, &TabletContext = LastPacketData.TabletContext]
+							{
+								return bIsSet && TabletContext && (TabletContext->GetSupportedProperties() & ETabletSupportedProperties::DeviceContactID)
+								       != ETabletSupportedProperties::None
+									       ? EVisibility::Visible
+									       : EVisibility::Collapsed;
+							})
+							.Text_Lambda([&DeviceContactID = LastPacketData.Packet.DeviceContactID]
+							{
+								return FText::Format(LOCTEXT("DeviceContactID", "Device Contact ID: {0}"), DeviceContactID);
+							})
+						]
+					]
+				]
+				+ SSplitter::Slot()
+				.Value(0.2f)
+				[
+					SNew(SMultiLineEditableText)
+					.IsReadOnly(true)
+					.Text_Lambda([&DebugMessages = DebugMessages] { return FText::FromString(DebugMessages); })
+					.VScrollBar(SNew(SScrollBar).Orientation(Orient_Vertical).AlwaysShowScrollbar(true).Thickness(10))
+				]
+			];
+	}
+
+	void SStylusInputDebugWidget::AcquireStylusInput()
+	{
+		if (StylusInput)
+		{
+			LogWarning("Stylus input instance has already been acquired.");
+			return;
+		}
+
+		check(FSlateApplication::IsInitialized());
+		const TSharedPtr<SWindow> Window = FSlateApplication::Get().FindWidgetWindow(AsShared());
+
+		if (!Window.IsValid())
+		{
+			LogError("Could not find widget window; stylus input instance has not been acquired.");
+			return;
+		}
+
+		StylusInput = CreateInstance(*Window);
+		if (!StylusInput)
+		{
+			LogError("Could not acquire stylus input instance.");
+		}
+	}
+
+	void SStylusInputDebugWidget::ReleaseStylusInput()
+	{
+		if (EventHandler)
+		{
+			LogWarning("Event handler is still registered.");
+		}
+
+		if (StylusInput)
+		{
+			if (!ReleaseInstance(StylusInput))
+			{
+				LogError("Failed to release stylus input for StylusInput Debug Widget.");
+			}
+
+			StylusInput = nullptr;
+		}
+	}
+
+	void SStylusInputDebugWidget::RegisterEventHandler()
+	{
+		if (!StylusInput)
+		{
+			LogWarning("Cannot register event handler since stylus input is unavailable.");
+			return;
+		}
+
+		if (EventHandler)
+		{
+			LogWarning("Event handler is not null; please unregister event handler before trying to registering it again.");
+		}
+
+		if (EventHandlerThread == EEventHandlerThread::Asynchronous)
+		{
+			EventHandler = MakeUnique<FDebugEventHandlerAsynchronous>(FOnPacketCallback::CreateRaw(this, &SStylusInputDebugWidget::OnPacket),
+			                                                          FOnDebugEventCallback::CreateRaw(this, &SStylusInputDebugWidget::OnDebugEvent));
+		}
+		else
+		{
+			EventHandler = MakeUnique<FDebugEventHandlerOnGameThread>(FOnPacketCallback::CreateRaw(this, &SStylusInputDebugWidget::OnPacket),
+			                                                          FOnDebugEventCallback::CreateRaw(this, &SStylusInputDebugWidget::OnDebugEvent));
+		}
+
+		if (StylusInput->AddEventHandler(EventHandler.Get(), EventHandlerThread))
+		{
+			LogVerbose(EventHandlerThread == EEventHandlerThread::Asynchronous
+				           ? "Registered event handler on asynchronous thread."
+				           : "Registered event handler on game thread.");
+		}
+		else
+		{
+			LogError("Failed to register event handler.");
+		}
+	}
+
+	void SStylusInputDebugWidget::UnregisterEventHandler()
+	{
+		if (!EventHandler)
+		{
+			LogWarning("Cannot unregister event handler since it is invalid.");
+			return;
+		}
+
+		if (StylusInput)
+		{
+			if (StylusInput->RemoveEventHandler(EventHandler.Get()))
+			{
+				LogVerbose("Unregistered event handler for StylusInput Debug Widget.");
+			}
+			else
+			{
+				LogError("Failed to unregister event handler for StylusInput Debug Widget.");
+			}
+		}
+		else
+		{
+			LogWarning("Cannot unregister event handler since stylus input is unavailable.");
+		}
+
+		EventHandler.Reset();
+	}
+
+	FDebugEventHandlerAsynchronous::FDebugEventHandlerAsynchronous(FOnPacketCallback&& OnPacketCallback, FOnDebugEventCallback&& OnDebugEventCallback)
+		: OnPacketCallback(MoveTemp(OnPacketCallback)), OnDebugEventCallback(MoveTemp(OnDebugEventCallback))
+	{
+		check(this->OnPacketCallback.IsBound());
+		check(this->OnDebugEventCallback.IsBound());
+	}
+
+	void FDebugEventHandlerAsynchronous::OnPacket(const FStylusInputPacket& Packet)
+	{
+		PacketQueue.Enqueue(Packet);
+	}
+
+	void FDebugEventHandlerAsynchronous::OnDebugEvent(const FString& Message)
+	{
+		DebugEventQueue.Enqueue(Message);
+	}
+
+	void FDebugEventHandlerAsynchronous::Tick(float DeltaTime)
+	{
+		FStylusInputPacket Packet;
+		while (PacketQueue.Dequeue(Packet))
+		{
+			OnPacketCallback.Execute(Packet);
+		}
+
+		FString Message;
+		while (DebugEventQueue.Dequeue(Message))
+		{
+			OnDebugEventCallback.Execute(Message);
+		}
+	}
+
+	FDebugEventHandlerOnGameThread::FDebugEventHandlerOnGameThread(FOnPacketCallback&& OnPacketCallback, FOnDebugEventCallback&& OnDebugEventCallback)
+		: OnPacketCallback(MoveTemp(OnPacketCallback)), OnDebugEventCallback(MoveTemp(OnDebugEventCallback))
+	{
+		check(this->OnPacketCallback.IsBound());
+		check(this->OnDebugEventCallback.IsBound());
+	}
+
+	void FDebugEventHandlerOnGameThread::OnPacket(const FStylusInputPacket& Packet)
+	{
+		OnPacketCallback.Execute(Packet);
+	}
+
+	void FDebugEventHandlerOnGameThread::OnDebugEvent(const FString& Message)
+	{
+		OnDebugEventCallback.Execute(Message);
+	}
+
+	SStylusInputDebugWidget::~SStylusInputDebugWidget()
+	{
+		UnregisterEventHandler();
+		ReleaseStylusInput();
+	}
+
+	void SStylusInputDebugWidget::OnPacket(const FStylusInputPacket& Packet)
+	{
+		LastPacketData.bIsSet = true;
+		LastPacketData.Packet = Packet;
+
+		if (!LastPacketData.TabletContext || (LastPacketData.TabletContext && LastPacketData.TabletContext->GetID() != Packet.TabletContextID))
+		{
+			LastPacketData.TabletContext = GetTabletContext(Packet.TabletContextID);
+		}
+
+		if (!LastPacketData.StylusInfo || (LastPacketData.StylusInfo && LastPacketData.StylusInfo->GetID() != Packet.CursorID))
+		{
+			LastPacketData.StylusInfo = GetStylusInfo(Packet.CursorID);
+		}
+
+		if (PaintWidget)
+		{
+			PaintWidget->Add(Packet);
+		}
+	}
+
+	void SStylusInputDebugWidget::OnDebugEvent(const FString& Message)
+	{
+		if (DebugMessages.IsEmpty())
+		{
+			DebugMessages = Message;
+		}
+		else
+		{
+			DebugMessages += "\n" + Message;
+		}
+	}
+
+	TSharedRef<SWidget> SStylusInputDebugWidget::GetEventHandlerThreadMenu()
+	{
+		FMenuBuilder MenuBuilder(true, nullptr);
+
+		MenuBuilder.AddMenuEntry(LOCTEXT("Asynchronous", "Asynchronous"),
+		                         LOCTEXT("AsynchronousTooltip", "Event handler is evaluated on a dedicated thread"),
+		                         FSlateIcon(),
+		                         FUIAction(FExecuteAction::CreateLambda([this]
+		                                   {
+			                                   SetEventHandlerThread(EEventHandlerThread::Asynchronous);
+		                                   }),
+		                                   FCanExecuteAction(),
+		                                   FIsActionChecked::CreateLambda([&EventHandlerThread = EventHandlerThread]
+		                                   {
+			                                   return EventHandlerThread == EEventHandlerThread::Asynchronous;
+		                                   })),
+		                         NAME_None,
+		                         EUserInterfaceActionType::RadioButton);
+
+		MenuBuilder.AddMenuEntry(LOCTEXT("GameThread", "On Game Thread"),
+		                         LOCTEXT("GameThreadTooltip", "Event handler is evaluated on the game thread"),
+		                         FSlateIcon(),
+		                         FUIAction(FExecuteAction::CreateLambda([this]
+		                                   {
+			                                   SetEventHandlerThread(EEventHandlerThread::OnGameThread);
+		                                   }),
+		                                   FCanExecuteAction(),
+		                                   FIsActionChecked::CreateLambda([&EventHandlerThread = EventHandlerThread]
+		                                   {
+			                                   return EventHandlerThread == EEventHandlerThread::OnGameThread;
+		                                   })),
+		                         NAME_None,
+		                         EUserInterfaceActionType::RadioButton);
+
+		return MenuBuilder.MakeWidget();
+	}
+
+	void SStylusInputDebugWidget::SetEventHandlerThread(EEventHandlerThread InEventHandlerThread)
+	{
+		if (EventHandlerThread != InEventHandlerThread)
+		{
+			UnregisterEventHandler();
+
+			EventHandlerThread = InEventHandlerThread;
+
+			RegisterEventHandler();
+		}
+	}
+
+	const IStylusInputTabletContext* SStylusInputDebugWidget::GetTabletContext(uint32 TabletContextID)
+	{
+		if (!StylusInput)
+		{
+			return nullptr;
+		}
+
+		const TSharedPtr<IStylusInputTabletContext>* TabletContext = TabletContexts.Find(TabletContextID);
+		if (!TabletContext)
+		{
+			if (const TSharedPtr<IStylusInputTabletContext>& NewTabletContext = StylusInput->GetTabletContext(TabletContextID))
+			{
+				TabletContext = &TabletContexts.Emplace(TabletContextID, NewTabletContext);
+			}
+		}
+
+		return TabletContext ? TabletContext->Get() : nullptr;
+	}
+
+	const IStylusInputStylusInfo* SStylusInputDebugWidget::GetStylusInfo(uint32 StylusID)
+	{
+		if (!StylusInput)
+		{
+			return nullptr;
+		}
+
+		const TSharedPtr<IStylusInputStylusInfo>* StylusInfo = StylusInfos.Find(StylusID);
+		if (!StylusInfo)
+		{
+			if (const TSharedPtr<IStylusInputStylusInfo>& NewStylusInfo = StylusInput->GetStylusInfo(StylusID))
+			{
+				StylusInfo = &StylusInfos.Emplace(StylusID, NewStylusInfo);
+			}
+			else
+			{
+				return nullptr;
+			}
+		}
+
+		return StylusInfo->Get();
+	}
+}
+
+#undef LOCTEXT_NAMESPACE
diff --git a/Engine/Plugins/Editor/StylusInput/Source/StylusInputDebugWidget/Private/StylusInputDebugWidget.h b/Engine/Plugins/Editor/StylusInput/Source/StylusInputDebugWidget/Private/StylusInputDebugWidget.h
new file mode 100644
--- /dev/null
+++ b/Engine/Plugins/Editor/StylusInput/Source/StylusInputDebugWidget/Private/StylusInputDebugWidget.h
@@ -0,0 +1,101 @@
+﻿// Copyright Epic Games, Inc. All Rights Reserved.
+
+#pragma once
+
+#include <StylusInput.h>
+#include <Widgets/SCompoundWidget.h>
+
+#include "StylusInputDebugPaintWidget.h"
+#include "TickableEditorObject.h"
+#include "Containers/SpscQueue.h"
+
+namespace UE::StylusInput::DebugWidget
+{
+	DECLARE_DELEGATE_OneParam(FOnPacketCallback, const FStylusInputPacket&);
+	DECLARE_DELEGATE_OneParam(FOnDebugEventCallback, const FString&);
+
+	class FDebugEventHandlerAsynchronous final : public IStylusInputEventHandler, FTickableEditorObject
+	{
+	public:
+		FDebugEventHandlerAsynchronous(FOnPacketCallback&& OnPacketCallback, FOnDebugEventCallback&& OnDebugEventCallback);
+
+		virtual FString GetName() override { return "DebugEventHandlerAsynchronous"; }
+
+		virtual void OnPacket(const FStylusInputPacket& Packet) override;
+		virtual void OnDebugEvent(const FString& Message) override;
+
+		virtual void Tick(float DeltaTime) override;
+		virtual TStatId GetStatId() const override { RETURN_QUICK_DECLARE_CYCLE_STAT(StylusInput_DebugEventHandlerAsynchronous, STATGROUP_Tickables); }
+
+	private:
+		FOnPacketCallback OnPacketCallback;
+		FOnDebugEventCallback OnDebugEventCallback;
+
+		TSpscQueue<FStylusInputPacket> PacketQueue;
+		TSpscQueue<FString> DebugEventQueue;
+	};
+
+	class FDebugEventHandlerOnGameThread final : public IStylusInputEventHandler
+	{
+	public:
+		FDebugEventHandlerOnGameThread(FOnPacketCallback&& OnPacketCallback, FOnDebugEventCallback&& OnDebugEventCallback);
+
+		virtual FString GetName() override { return "DebugEventHandlerOnGameThread"; }
+
+		virtual void OnPacket(const FStylusInputPacket& Packet) override;
+		virtual void OnDebugEvent(const FString& Message) override;
+
+	private:
+		FOnPacketCallback OnPacketCallback;
+		FOnDebugEventCallback OnDebugEventCallback;
+	};
+
+	class SStylusInputDebugWidget : public SCompoundWidget
+	{
+	public:
+		SLATE_BEGIN_ARGS(SStylusInputDebugWidget)
+			{
+			}
+		SLATE_END_ARGS()
+
+		virtual ~SStylusInputDebugWidget() override;
+
+		void Construct(const FArguments& Args);
+
+	private:
+		void AcquireStylusInput();
+		void ReleaseStylusInput();
+		void RegisterEventHandler();
+		void UnregisterEventHandler();
+
+		const IStylusInputTabletContext* GetTabletContext(uint32 TabletContextID);
+		const IStylusInputStylusInfo* GetStylusInfo(uint32 StylusID);
+
+		void OnPacket(const FStylusInputPacket& Packet);
+		void OnDebugEvent(const FString& Message);
+
+		TSharedRef<SWidget> GetEventHandlerThreadMenu();
+		void SetEventHandlerThread(EEventHandlerThread InEventHandlerThread);
+		EEventHandlerThread EventHandlerThread = EEventHandlerThread::OnGameThread;
+
+		IStylusInputInstance* StylusInput = nullptr;
+		TUniquePtr<IStylusInputEventHandler> EventHandler;
+
+		TSharedPtr<SStylusInputDebugPaintWidget> PaintWidget;
+
+		TMap<uint32, TSharedPtr<IStylusInputTabletContext>> TabletContexts;
+		TMap<uint32, TSharedPtr<IStylusInputStylusInfo>> StylusInfos;
+
+		FString DebugMessages;
+
+		struct FLastPacketData
+		{
+			bool bIsSet = false;
+			FStylusInputPacket Packet;
+			const IStylusInputTabletContext* TabletContext = nullptr;
+			const IStylusInputStylusInfo* StylusInfo = nullptr;
+		};
+
+		FLastPacketData LastPacketData;
+	};
+}
diff --git a/Engine/Plugins/Editor/StylusInput/Source/StylusInputDebugWidget/Private/StylusInputDebugWidgetModule.cpp b/Engine/Plugins/Editor/StylusInput/Source/StylusInputDebugWidget/Private/StylusInputDebugWidgetModule.cpp
new file mode 100644
--- /dev/null
+++ b/Engine/Plugins/Editor/StylusInput/Source/StylusInputDebugWidget/Private/StylusInputDebugWidgetModule.cpp
@@ -0,0 +1,93 @@
+﻿// Copyright Epic Games, Inc. All Rights Reserved.
+
+#include <WorkspaceMenuStructure.h>
+#include <WorkspaceMenuStructureModule.h>
+#include <Modules/ModuleManager.h>
+#include <Widgets/Docking/SDockTab.h>
+
+#include "StylusInputDebugWidget.h"
+
+#define LOCTEXT_NAMESPACE "StylusInputDebugWidgetModule"
+
+namespace UE::StylusInput::DebugWidget
+{
+	class FStylusInputDebugWidgetModule : public IModuleInterface
+	{
+	public:
+		virtual void StartupModule() override;
+		virtual void ShutdownModule() override;
+
+		void RegisterTabSpawners(const TSharedPtr<FWorkspaceItem>& WorkspaceGroup);
+		void UnregisterTabSpawners();
+
+	private:
+		TSharedRef<SDockTab> MakeDebugWidgetTab(const FSpawnTabArgs&);
+		TSharedRef<SWidget> GetDebugWidget();
+
+		bool bHasRegisteredTabSpawners = false;
+		TWeakPtr<SStylusInputDebugWidget> DebugWidgetPtr;
+	};
+
+	void FStylusInputDebugWidgetModule::StartupModule()
+	{
+		RegisterTabSpawners(nullptr);
+	}
+
+	void FStylusInputDebugWidgetModule::ShutdownModule()
+	{
+		UnregisterTabSpawners();
+	}
+
+	void FStylusInputDebugWidgetModule::RegisterTabSpawners(const TSharedPtr<FWorkspaceItem>& WorkspaceGroup)
+	{
+		if (bHasRegisteredTabSpawners)
+		{
+			UnregisterTabSpawners();
+		}
+
+		FTabSpawnerEntry& DebugWidgetSpawnerEntry = FGlobalTabmanager::Get()->RegisterNomadTabSpawner("StylusInputDebugWidget",
+			FOnSpawnTab::CreateRaw(this, &FStylusInputDebugWidgetModule::MakeDebugWidgetTab))
+				.SetDisplayName(LOCTEXT("DebugWidgetTitle", "Stylus Input Debug"))
+				.SetTooltipText(LOCTEXT("DebugWidgetTooltip", "Open a debug widget to verify stylus input event handling."))
+				.SetGroup(WorkspaceMenu::GetMenuStructure().GetDeveloperToolsDebugCategory())
+				.SetIcon(FSlateIcon(FAppStyle::GetAppStyleSetName(), "StylusInputDebug.TabIcon"));
+
+		if (WorkspaceGroup.IsValid())
+		{
+			DebugWidgetSpawnerEntry.SetGroup(WorkspaceGroup.ToSharedRef());
+		}
+
+		bHasRegisteredTabSpawners = true;
+	}
+
+	void FStylusInputDebugWidgetModule::UnregisterTabSpawners()
+	{
+		bHasRegisteredTabSpawners = false;
+
+		FGlobalTabmanager::Get()->UnregisterNomadTabSpawner("StylusInputDebugWidget");
+	}
+
+	TSharedRef<SDockTab> FStylusInputDebugWidgetModule::MakeDebugWidgetTab(const FSpawnTabArgs&)
+	{
+		TSharedRef<SDockTab> Tab = SNew(SDockTab).TabRole(NomadTab);
+		Tab->SetContent(GetDebugWidget());
+		return Tab;
+	}
+
+	TSharedRef<SWidget> FStylusInputDebugWidgetModule::GetDebugWidget()
+	{
+		TSharedPtr<SStylusInputDebugWidget> DebugWidget = DebugWidgetPtr.Pin();
+
+		if (!DebugWidget.IsValid())
+		{
+			DebugWidget = SNew(SStylusInputDebugWidget);
+			DebugWidgetPtr = DebugWidget;
+		}
+
+		return DebugWidget.ToSharedRef();
+	}
+}
+
+#undef LOCTEXT_NAMESPACE
+
+IMPLEMENT_MODULE(UE::StylusInput::DebugWidget::FStylusInputDebugWidgetModule, StylusInputDebugWidget);
diff --git a/Engine/Plugins/Editor/StylusInput/Source/StylusInputDebugWidget/StylusInputDebugWidget.Build.cs b/Engine/Plugins/Editor/StylusInput/Source/StylusInputDebugWidget/StylusInputDebugWidget.Build.cs
new file mode 100644
--- /dev/null
+++ b/Engine/Plugins/Editor/StylusInput/Source/StylusInputDebugWidget/StylusInputDebugWidget.Build.cs
@@ -0,0 +1,23 @@
+﻿// Copyright Epic Games, Inc. All Rights Reserved.
+
+namespace UnrealBuildTool.Rules
+{
+	public class StylusInputDebugWidget : ModuleRules
+	{
+		public StylusInputDebugWidget(ReadOnlyTargetRules Target) : base(Target)
+		{
+			PrivateDependencyModuleNames.AddRange(
+				new string[]
+				{
+					"Core",
+					"Engine",
+					"Slate",
+					"SlateCore",
+					"StylusInput",
+					"UnrealEd",
+					"WorkspaceMenuStructure",
+				}
+			);
+		}
+	}
+}
diff --git a/Engine/Plugins/Editor/StylusInput/StylusInput.uplugin b/Engine/Plugins/Editor/StylusInput/StylusInput.uplugin
--- a/Engine/Plugins/Editor/StylusInput/StylusInput.uplugin
+++ b/Engine/Plugins/Editor/StylusInput/StylusInput.uplugin
@@ -1,9 +1,9 @@
 {
 	"FileVersion" : 3,
-	"Version" : 1,
-	"VersionName" : "1.0",
+	"Version" : 2,
+	"VersionName" : "2.0",
 	"FriendlyName" : "Stylus & Tablet Plugin",
-	"Description" : "Support for advanced stylus and tablet inputs such as pressure, stylus & tablet buttons, and pen angles.",
+	"Description" : "Support for advanced stylus and tablet inputs such as pressure, stylus and tablet buttons, and pen angles.",
 	"Category" : "Input Devices",
 	"CreatedBy" : "Epic Games, Inc.",
 	"CreatedByURL" : "https://epicgames.com",
@@ -20,6 +20,10 @@
 			"Name" : "StylusInput",
 			"Type" : "Editor",
 			"LoadingPhase" : "Default"
+		},
+		{
+			"Name": "StylusInputDebugWidget",
+			"Type": "Editor"
 		}
 	]
-}
\ No newline at end of file
+}
