StylusInput: Fix missing platform includes

diff --git a/Engine/Plugins/Editor/StylusInput/Source/StylusInput/Private/DeprecatedWindowsStylusInputInterface.h b/Engine/Plugins/Editor/StylusInput/Source/StylusInput/Private/DeprecatedWindowsStylusInputInterface.h
--- a/Engine/Plugins/Editor/StylusInput/Source/StylusInput/Private/DeprecatedWindowsStylusInputInterface.h
+++ b/Engine/Plugins/Editor/StylusInput/Source/StylusInput/Private/DeprecatedWindowsStylusInputInterface.h
@@ -2,6 +2,8 @@
 
 #pragma once
 
+#include <HAL/Platform.h>
+
 #if PLATFORM_WINDOWS
 
 #include "IStylusInputModule.h"
diff --git a/Engine/Plugins/Editor/StylusInput/Source/StylusInput/Private/Windows/WindowsStylusInputInstance.h b/Engine/Plugins/Editor/StylusInput/Source/StylusInput/Private/Windows/WindowsStylusInputInstance.h
--- a/Engine/Plugins/Editor/StylusInput/Source/StylusInput/Private/Windows/WindowsStylusInputInstance.h
+++ b/Engine/Plugins/Editor/StylusInput/Source/StylusInput/Private/Windows/WindowsStylusInputInstance.h
@@ -2,6 +2,8 @@
 
 #pragma once
 
+#include <HAL/Platform.h>
+
 #if PLATFORM_WINDOWS
 
 #include <StylusInput.h>
diff --git a/Engine/Plugins/Editor/StylusInput/Source/StylusInput/Private/Windows/WindowsStylusInputInterface.h b/Engine/Plugins/Editor/StylusInput/Source/StylusInput/Private/Windows/WindowsStylusInputInterface.h
--- a/Engine/Plugins/Editor/StylusInput/Source/StylusInput/Private/Windows/WindowsStylusInputInterface.h
+++ b/Engine/Plugins/Editor/StylusInput/Source/StylusInput/Private/Windows/WindowsStylusInputInterface.h
@@ -2,6 +2,8 @@
 
 #pragma once
 
+#include <HAL/Platform.h>
+
 #if PLATFORM_WINDOWS
 
 #include <StylusInputInterface.h>
diff --git a/Engine/Plugins/Editor/StylusInput/Source/StylusInput/Private/Windows/WindowsStylusInputPlatformAPI.h b/Engine/Plugins/Editor/StylusInput/Source/StylusInput/Private/Windows/WindowsStylusInputPlatformAPI.h
--- a/Engine/Plugins/Editor/StylusInput/Source/StylusInput/Private/Windows/WindowsStylusInputPlatformAPI.h
+++ b/Engine/Plugins/Editor/StylusInput/Source/StylusInput/Private/Windows/WindowsStylusInputPlatformAPI.h
@@ -2,6 +2,8 @@
 
 #pragma once
 
+#include <HAL/Platform.h>
+
 #if PLATFORM_WINDOWS
 
 #include <Containers/Array.h>
diff --git a/Engine/Plugins/Editor/StylusInput/Source/StylusInput/Private/Windows/WindowsStylusInputPluginAsync.h b/Engine/Plugins/Editor/StylusInput/Source/StylusInput/Private/Windows/WindowsStylusInputPluginAsync.h
--- a/Engine/Plugins/Editor/StylusInput/Source/StylusInput/Private/Windows/WindowsStylusInputPluginAsync.h
+++ b/Engine/Plugins/Editor/StylusInput/Source/StylusInput/Private/Windows/WindowsStylusInputPluginAsync.h
@@ -2,6 +2,8 @@
 
 #pragma once
 
+#include <HAL/Platform.h>
+
 #if PLATFORM_WINDOWS
 
 #include "WindowsStylusInputPluginBase.h"
diff --git a/Engine/Plugins/Editor/StylusInput/Source/StylusInput/Private/Windows/WindowsStylusInputPluginBase.h b/Engine/Plugins/Editor/StylusInput/Source/StylusInput/Private/Windows/WindowsStylusInputPluginBase.h
--- a/Engine/Plugins/Editor/StylusInput/Source/StylusInput/Private/Windows/WindowsStylusInputPluginBase.h
+++ b/Engine/Plugins/Editor/StylusInput/Source/StylusInput/Private/Windows/WindowsStylusInputPluginBase.h
@@ -2,6 +2,8 @@
 
 #pragma once
 
+#include <HAL/Platform.h>
+
 #if PLATFORM_WINDOWS
 
 #include <StylusInputPacket.h>
diff --git a/Engine/Plugins/Editor/StylusInput/Source/StylusInput/Private/Windows/WindowsStylusInputPluginSync.h b/Engine/Plugins/Editor/StylusInput/Source/StylusInput/Private/Windows/WindowsStylusInputPluginSync.h
--- a/Engine/Plugins/Editor/StylusInput/Source/StylusInput/Private/Windows/WindowsStylusInputPluginSync.h
+++ b/Engine/Plugins/Editor/StylusInput/Source/StylusInput/Private/Windows/WindowsStylusInputPluginSync.h
@@ -2,6 +2,8 @@
 
 #pragma once
 
+#include <HAL/Platform.h>
+
 #if PLATFORM_WINDOWS
 
 #include "WindowsStylusInputPluginBase.h"
diff --git a/Engine/Plugins/Editor/StylusInput/Source/StylusInput/Private/Windows/WindowsStylusInputStats.h b/Engine/Plugins/Editor/StylusInput/Source/StylusInput/Private/Windows/WindowsStylusInputStats.h
--- a/Engine/Plugins/Editor/StylusInput/Source/StylusInput/Private/Windows/WindowsStylusInputStats.h
+++ b/Engine/Plugins/Editor/StylusInput/Source/StylusInput/Private/Windows/WindowsStylusInputStats.h
@@ -2,9 +2,10 @@
 
 #pragma once
 
+#include <HAL/Platform.h>
+
 #if PLATFORM_WINDOWS
 
-// #include <HAL/Platform.h>
 #include "Windows/WindowsPlatformTime.h"
 
 namespace UE::StylusInput::Private::Windows
diff --git a/Engine/Plugins/Editor/StylusInput/Source/StylusInput/Private/Windows/WindowsStylusInputTabletContext.h b/Engine/Plugins/Editor/StylusInput/Source/StylusInput/Private/Windows/WindowsStylusInputTabletContext.h
--- a/Engine/Plugins/Editor/StylusInput/Source/StylusInput/Private/Windows/WindowsStylusInputTabletContext.h
+++ b/Engine/Plugins/Editor/StylusInput/Source/StylusInput/Private/Windows/WindowsStylusInputTabletContext.h
@@ -2,6 +2,8 @@
 
 #pragma once
 
+#include <HAL/Platform.h>
+
 #if PLATFORM_WINDOWS
 
 #include <StylusInputPacket.h>
diff --git a/Engine/Plugins/Editor/StylusInput/Source/StylusInput/Private/Windows/WindowsStylusInputUtils.h b/Engine/Plugins/Editor/StylusInput/Source/StylusInput/Private/Windows/WindowsStylusInputUtils.h
--- a/Engine/Plugins/Editor/StylusInput/Source/StylusInput/Private/Windows/WindowsStylusInputUtils.h
+++ b/Engine/Plugins/Editor/StylusInput/Source/StylusInput/Private/Windows/WindowsStylusInputUtils.h
@@ -2,6 +2,8 @@
 
 #pragma once
 
+#include <HAL/Platform.h>
+
 #if PLATFORM_WINDOWS
 
 #include <Containers/UnrealString.h>
