StylusInput: Deprecate old API

diff --git a/Engine/Plugins/Editor/StylusInput/Source/StylusInput/Private/DeprecatedStylusInputSubsystem.cpp b/Engine/Plugins/Editor/StylusInput/Source/StylusInput/Private/DeprecatedStylusInputSubsystem.cpp
--- a/Engine/Plugins/Editor/StylusInput/Source/StylusInput/Private/DeprecatedStylusInputSubsystem.cpp
+++ b/Engine/Plugins/Editor/StylusInput/Source/StylusInput/Private/DeprecatedStylusInputSubsystem.cpp
@@ -5,6 +5,8 @@
 
 #include "IStylusInputModule.h"
 
+PRAGMA_DISABLE_DEPRECATION_WARNINGS
+
 #define LOCTEXT_NAMESPACE "FStylusInputModule"
 
 static const FName StylusInputDebugTabName = FName("StylusInputDebug");
@@ -100,3 +102,5 @@ void UStylusInputSubsystem::Tick(float DeltaTime)
 }
 
 #undef LOCTEXT_NAMESPACE
+
+PRAGMA_ENABLE_DEPRECATION_WARNINGS
diff --git a/Engine/Plugins/Editor/StylusInput/Source/StylusInput/Private/DeprecatedWindowsStylusInputInterface.cpp b/Engine/Plugins/Editor/StylusInput/Source/StylusInput/Private/DeprecatedWindowsStylusInputInterface.cpp
--- a/Engine/Plugins/Editor/StylusInput/Source/StylusInput/Private/DeprecatedWindowsStylusInputInterface.cpp
+++ b/Engine/Plugins/Editor/StylusInput/Source/StylusInput/Private/DeprecatedWindowsStylusInputInterface.cpp
@@ -21,7 +21,9 @@ void FDeprecatedWindowsStylusInputInterface::Tick()
 	for (const FDeprecatedStylusInputDevice& TabletContext : *TabletContexts)
 	{
 		// don't change focus if any stylus is down
+		PRAGMA_DISABLE_DEPRECATION_WARNINGS
 		if (TabletContext.GetCurrentState().IsStylusDown())
+		PRAGMA_ENABLE_DEPRECATION_WARNINGS
 		{
 			return;
 		}
@@ -55,10 +57,12 @@ int32 FDeprecatedWindowsStylusInputInterface::NumInputDevices() const
 	return TabletContexts->Num();
 }
 
+PRAGMA_DISABLE_DEPRECATION_WARNINGS
 IStylusInputDevice* FDeprecatedWindowsStylusInputInterface::GetInputDevice(int32 Index) const
 {
 	return 0 <= Index && Index < TabletContexts->Num() ? &(*TabletContexts)[Index] : nullptr;
 }
+PRAGMA_ENABLE_DEPRECATION_WARNINGS
 
 void FDeprecatedWindowsStylusInputInterface::CreateStylusInputInstance(SWindow* Window)
 {	
@@ -70,6 +74,7 @@ void FDeprecatedWindowsStylusInputInterface::RemoveStylusInputInstance(const TSh
 	StylusInputInstances.Remove(&Window.Get());
 }
 
+PRAGMA_DISABLE_DEPRECATION_WARNINGS
 FStylusState ToDeprecatedStylusState(const UE::StylusInput::FStylusInputPacket& Packet)
 {
 	return {
@@ -79,7 +84,9 @@ FStylusState ToDeprecatedStylusState(const UE::StylusInput::FStylusInputPacket&
 		(Packet.PenStatus & UE::StylusInput::EPenStatus::CursorIsInverted) != UE::StylusInput::EPenStatus::None
 	};
 }
+PRAGMA_ENABLE_DEPRECATION_WARNINGS
 
+PRAGMA_DISABLE_DEPRECATION_WARNINGS
 FDeprecatedWindowsStylusInputInterface::FDeprecatedStylusInputDevice::FDeprecatedStylusInputDevice(uint32 TabletContextId, const TSharedPtr<UE::StylusInput::IStylusInputTabletContext>& TabletContext)
 	: TabletContextId(TabletContextId)
 {
@@ -133,12 +140,15 @@ FDeprecatedWindowsStylusInputInterface::FDeprecatedStylusInputDevice::FDeprecate
 		}
 	}
 }
+PRAGMA_ENABLE_DEPRECATION_WARNINGS
 
 void FDeprecatedWindowsStylusInputInterface::FDeprecatedStylusInputDevice::Tick()
 {
+	PRAGMA_DISABLE_DEPRECATION_WARNINGS
 	PreviousState = CurrentState;
 	CurrentState = ToDeprecatedStylusState(LastPacket);
 	Dirty = false;
+	PRAGMA_ENABLE_DEPRECATION_WARNINGS
 }
 
 FDeprecatedWindowsStylusInputInterface::FStylusInputEventHandler::FStylusInputEventHandler(UE::StylusInput::IStylusInputInstance* Instance, TArray<FDeprecatedStylusInputDevice>& TabletContexts)
@@ -200,7 +210,10 @@ FDeprecatedWindowsStylusInputInterface::FStylusInputInstanceWrapper::~FStylusInp
 	}
 }
 
+
+PRAGMA_DISABLE_DEPRECATION_WARNINGS
 TSharedPtr<IStylusInputInterfaceInternal> CreateStylusInputInterface()
+PRAGMA_ENABLE_DEPRECATION_WARNINGS
 {
 	return MakeShared<FDeprecatedWindowsStylusInputInterface>();
 }
diff --git a/Engine/Plugins/Editor/StylusInput/Source/StylusInput/Private/DeprecatedWindowsStylusInputInterface.h b/Engine/Plugins/Editor/StylusInput/Source/StylusInput/Private/DeprecatedWindowsStylusInputInterface.h
--- a/Engine/Plugins/Editor/StylusInput/Source/StylusInput/Private/DeprecatedWindowsStylusInputInterface.h
+++ b/Engine/Plugins/Editor/StylusInput/Source/StylusInput/Private/DeprecatedWindowsStylusInputInterface.h
@@ -15,7 +15,11 @@ class SWindow;
 /**
  * An implementation of the deprecated @see UStylusInputSubsystem interface by means of using the new interface and Windows implementation.
  */
-class FDeprecatedWindowsStylusInputInterface : public IStylusInputInterfaceInternal, public TSharedFromThis<FDeprecatedWindowsStylusInputInterface>
+class  FDeprecatedWindowsStylusInputInterface
+	PRAGMA_DISABLE_DEPRECATION_WARNINGS
+	: public IStylusInputInterfaceInternal
+	PRAGMA_ENABLE_DEPRECATION_WARNINGS
+	, public TSharedFromThis<FDeprecatedWindowsStylusInputInterface>
 {
 public:
 	FDeprecatedWindowsStylusInputInterface();
@@ -23,12 +27,15 @@ public:
 
 	virtual void Tick() override;
 	virtual int32 NumInputDevices() const override;
+	PRAGMA_DISABLE_DEPRECATION_WARNINGS
 	virtual IStylusInputDevice* GetInputDevice(int32 Index) const override;
+	PRAGMA_ENABLE_DEPRECATION_WARNINGS
 
 private:
 	void CreateStylusInputInstance(SWindow* Window);
 	void RemoveStylusInputInstance(const TSharedRef<SWindow>& Window);
 
+	PRAGMA_DISABLE_DEPRECATION_WARNINGS
 	struct FDeprecatedStylusInputDevice : IStylusInputDevice
 	{
 		explicit FDeprecatedStylusInputDevice(uint32 TabletContextId, const TSharedPtr<UE::StylusInput::IStylusInputTabletContext>& TabletContext);
@@ -40,6 +47,7 @@ private:
 
 		UE::StylusInput::FStylusInputPacket LastPacket;
 	};
+	PRAGMA_ENABLE_DEPRECATION_WARNINGS
 
 	class FStylusInputEventHandler : public UE::StylusInput::IStylusInputEventHandler
 	{
diff --git a/Engine/Plugins/Editor/StylusInput/Source/StylusInput/Public/IStylusInputModule.h b/Engine/Plugins/Editor/StylusInput/Source/StylusInput/Public/IStylusInputModule.h
--- a/Engine/Plugins/Editor/StylusInput/Source/StylusInput/Public/IStylusInputModule.h
+++ b/Engine/Plugins/Editor/StylusInput/Source/StylusInput/Public/IStylusInputModule.h
@@ -19,7 +19,7 @@ DEFINE_LOG_CATEGORY_STATIC(LogStylusInput, Log, All);
 /**
  * Module to handle Wacom-style tablet input using styluses.
  */
-class STYLUSINPUT_API IStylusInputModule : public IModuleInterface
+class UE_DEPRECATED(5.5, "Please use the new API in StylusInput.h instead.") STYLUSINPUT_API IStylusInputModule : public IModuleInterface
 {
 public:
 
@@ -43,18 +43,20 @@ public:
 };
 
 // This is the interface that all platform-specific implementations must implement.
-class IStylusInputInterfaceInternal
+class UE_DEPRECATED(5.5, "Please use the new API in StylusInput.h instead.") IStylusInputInterfaceInternal
 {
 public:
 	virtual void Tick() = 0;
 
+	PRAGMA_DISABLE_DEPRECATION_WARNINGS
 	virtual IStylusInputDevice* GetInputDevice(int32 Index) const = 0;
+	PRAGMA_ENABLE_DEPRECATION_WARNINGS
 	virtual int32 NumInputDevices() const = 0;
 };
 
 
 UCLASS()
-class STYLUSINPUT_API UStylusInputSubsystem : 
+class UE_DEPRECATED(5.5, "Please use the new API's window-specific IStylusInputInstance in StylusInput.h instead.") STYLUSINPUT_API UStylusInputSubsystem : 
 	public UEditorSubsystem, 
 	public FTickableEditorObject
 {
@@ -65,22 +67,30 @@ public:
 	virtual void Deinitialize() override;
 
 	/** Retrieve the input device that is at the given index, or nullptr if not found. Corresponds to the StylusIndex in IStylusMessageHandler. */
+	PRAGMA_DISABLE_DEPRECATION_WARNINGS
 	const IStylusInputDevice* GetInputDevice(int32 Index) const;
+	PRAGMA_ENABLE_DEPRECATION_WARNINGS
 
 	/** Return the number of active input devices. */
 	int32 NumInputDevices() const; 
 
 	/** Add a message handler to receive messages from the stylus. */
+	PRAGMA_DISABLE_DEPRECATION_WARNINGS
 	void AddMessageHandler(IStylusMessageHandler& MessageHandler);
+	PRAGMA_ENABLE_DEPRECATION_WARNINGS
 
 	/** Remove a previously registered message handler. */
+	PRAGMA_DISABLE_DEPRECATION_WARNINGS
 	void RemoveMessageHandler(IStylusMessageHandler& MessageHandler);
+	PRAGMA_ENABLE_DEPRECATION_WARNINGS
 
 	// FTickableEditorObject implementation
 	virtual void Tick(float DeltaTime) override;
 	virtual TStatId GetStatId() const override { RETURN_QUICK_DECLARE_CYCLE_STAT(UStylusInputSubsystem, STATGROUP_Tickables); }
 
 private:
+	PRAGMA_DISABLE_DEPRECATION_WARNINGS
 	TSharedPtr<IStylusInputInterfaceInternal> InputInterface;
 	TArray<IStylusMessageHandler*> MessageHandlers;
+	PRAGMA_ENABLE_DEPRECATION_WARNINGS
 };
diff --git a/Engine/Plugins/Editor/StylusInput/Source/StylusInput/Public/IStylusState.h b/Engine/Plugins/Editor/StylusInput/Source/StylusInput/Public/IStylusState.h
--- a/Engine/Plugins/Editor/StylusInput/Source/StylusInput/Public/IStylusState.h
+++ b/Engine/Plugins/Editor/StylusInput/Source/StylusInput/Public/IStylusState.h
@@ -7,7 +7,7 @@
 /**
  * The types of stylus inputs that can be potentially supported by a stylus.
  */
-enum class EStylusInputType
+enum class UE_DEPRECATED(5.5, "Please use the new API's ETabletSupportedProperties in StylusInputTabletContext.h instead.") EStylusInputType
 {
 	Position,
 	Z,
@@ -22,7 +22,7 @@ enum class EStylusInputType
 /**
  * The current state of a single stylus, as sent by IStylusMessageHandler. 
  */
-class FStylusState
+class UE_DEPRECATED(5.5, "Please use the new API's FStylusInputPacket in StylusInputPacket.h instead.") FStylusState
 {
 public:
 
@@ -115,12 +115,14 @@ private:
 };
 
 /** An input device representing a stylus and its current state. */
-class IStylusInputDevice
+class UE_DEPRECATED(5.5, "Please use the new API in StylusInput.h instead.") IStylusInputDevice
 {
 public:
 
 	virtual ~IStylusInputDevice() {}
 
+	PRAGMA_DISABLE_DEPRECATION_WARNINGS
+	
 	/**
 	 * Get the current stylus state.
 	 */
@@ -135,6 +137,8 @@ public:
 	 * Get the supported inputs of this tablet.
 	 */
 	const TArray<EStylusInputType>& GetSupportedInputs() const { return SupportedInputs; }
+	
+	PRAGMA_ENABLE_DEPRECATION_WARNINGS
 
 	/** Update the input device. Not intended to be called externally. */ 
 	virtual void Tick() = 0;
@@ -143,9 +147,11 @@ public:
 	bool IsDirty() const { return Dirty; }
 
 protected:
+	PRAGMA_DISABLE_DEPRECATION_WARNINGS
 	FStylusState CurrentState;
 	FStylusState PreviousState;
 	TArray<EStylusInputType> SupportedInputs;
+	PRAGMA_ENABLE_DEPRECATION_WARNINGS
 	bool Dirty : 1;
 };
 
@@ -153,8 +159,10 @@ protected:
  * Interface to implement for classes that want to receive messages when a stylus state change occurs.
  * Will trigger once per frame.
  */
-class IStylusMessageHandler
+class UE_DEPRECATED(5.5, "Please use the new API's IStylusInputEventHandler in StylusInput.h instead.") IStylusMessageHandler
 {
 public:
+	PRAGMA_DISABLE_DEPRECATION_WARNINGS
 	virtual void OnStylusStateChanged(const FStylusState& NewState, int32 StylusIndex) = 0;
+	PRAGMA_ENABLE_DEPRECATION_WARNINGS
 };
diff --git a/Engine/Plugins/Editor/StylusInput/Source/StylusInput/Public/StylusInput.h b/Engine/Plugins/Editor/StylusInput/Source/StylusInput/Public/StylusInput.h
--- a/Engine/Plugins/Editor/StylusInput/Source/StylusInput/Public/StylusInput.h
+++ b/Engine/Plugins/Editor/StylusInput/Source/StylusInput/Public/StylusInput.h
@@ -61,7 +61,7 @@ namespace UE::StylusInput
 		 * @param Thread The type of thread the event handler calls are performed on.
 		 * @returns True if the event handler was added successfully.
 		 */
-		STYLUSINPUT_API virtual bool AddEventHandler(IStylusInputEventHandler* EventHandler, EEventHandlerThread Thread) = 0;
+		virtual bool AddEventHandler(IStylusInputEventHandler* EventHandler, EEventHandlerThread Thread) = 0;
 
 		/**
 		 * Removes an event handler from this stylus input instances, and frees up all associated resources.
@@ -70,7 +70,7 @@ namespace UE::StylusInput
 		 * @param EventHandler The event handler to be released.
 		 * @returns True if the event handler was removed successfully.
 		 */
-		STYLUSINPUT_API virtual bool RemoveEventHandler(IStylusInputEventHandler* EventHandler) = 0;
+		virtual bool RemoveEventHandler(IStylusInputEventHandler* EventHandler) = 0;
 
 		/**
 		 * Provides the tablet context information for a given tablet context ID.
@@ -82,7 +82,7 @@ namespace UE::StylusInput
 		 * @param TabletContextID The unique identifier of the tablet context as provided by this stylus input instance via an event handler callback. 
 		 * @returns Pointer to tablet context information or nullptr if the tablet context ID was invalid or an error occured. 
 		 */
-		STYLUSINPUT_API virtual const TSharedPtr<IStylusInputTabletContext> GetTabletContext(uint32 TabletContextID) = 0;
+		virtual const TSharedPtr<IStylusInputTabletContext> GetTabletContext(uint32 TabletContextID) = 0;
 
 		/**
 		 * Provides the stylus a.k.a. cursor or pen information for a given stylus ID.
@@ -94,7 +94,7 @@ namespace UE::StylusInput
 		 * @param StylusID The unique identifier of the stylus/cursor/pen as provided by this stylus input instance via an event handler callback. 
 		 * @returns Pointer to stylus information or nullptr if the stylus ID was invalid or an error occured. 
 		 */
-		STYLUSINPUT_API virtual const TSharedPtr<IStylusInputStylusInfo> GetStylusInfo(uint32 StylusID) = 0;
+		virtual const TSharedPtr<IStylusInputStylusInfo> GetStylusInfo(uint32 StylusID) = 0;
 
 		/**
 		 * Provides the approximate number of stylus input packets processed per second for the given thread type for diagnostic purposes.
@@ -103,7 +103,7 @@ namespace UE::StylusInput
 		 * @param Thread Selects the set of event handlers, i.e. number of packets processed on the game thread versus processed asynchronously.
 		 * @returns Approximate number of stylus input packets processed per second, or -1 if there is no valid data available.
 		 */
-		STYLUSINPUT_API virtual float GetPacketsPerSecond(EEventHandlerThread Thread) const { return -1.0f; }
+		virtual float GetPacketsPerSecond(EEventHandlerThread Thread) const { return -1.0f; }
 	};
 
 	/**
@@ -119,7 +119,7 @@ namespace UE::StylusInput
 		 *
 		 * @returns Name of the event handler.
 		 */
-		STYLUSINPUT_API virtual FString GetName() = 0;
+		virtual FString GetName() = 0;
 
 		/**
 		 * Callback for each packet processed by the stylus input instance.
@@ -128,15 +128,16 @@ namespace UE::StylusInput
 		 *
 		 * @param Packet The packet being processed.
 		 */
-		STYLUSINPUT_API virtual void OnPacket(const FStylusInputPacket& Packet) = 0;
+		virtual void OnPacket(const FStylusInputPacket& Packet) = 0;
 
 		/**
 		 * Callback for each debug event sent by the stylus input instance for diagnostic purposes.
 		 * This function is called on the game thread or asynchronously to the game thread depending on how the event handler was added to the stylus input
-		 * instance. The implementation of this function is responsible for any thread synchronization that might be necessary. 
+		 * instance. The implementation of this function is responsible for any thread synchronization that might be necessary.
+		 * The base class implementation of this function does nothing.
 		 *
 		 * @param Message The debug message sent by the stylus input instance.
 		 */
-		virtual void OnDebugEvent(const FString& Message) = 0;
+		virtual void OnDebugEvent(const FString& Message) {}
 	};
 }
diff --git a/Engine/Plugins/Editor/StylusInput/Source/StylusInput/Public/StylusInputTabletContext.h b/Engine/Plugins/Editor/StylusInput/Source/StylusInput/Public/StylusInputTabletContext.h
--- a/Engine/Plugins/Editor/StylusInput/Source/StylusInput/Public/StylusInputTabletContext.h
+++ b/Engine/Plugins/Editor/StylusInput/Source/StylusInput/Public/StylusInputTabletContext.h
@@ -64,7 +64,7 @@ namespace UE::StylusInput
 		 *
 		 * @returns The unique identifier for the tablet context for the respective stylus input instance.
 		 */
-		STYLUSINPUT_API virtual uint32 GetID() const = 0;
+		virtual uint32 GetID() const = 0;
 
 		/**
 		 * Returns the name of the tablet context.
@@ -72,7 +72,7 @@ namespace UE::StylusInput
 		 *
 		 * @returns The name of the tablet context.
 		 */
-		STYLUSINPUT_API virtual FString GetName() const = 0;
+		virtual FString GetName() const = 0;
 
 		/**
 		 * Returns the dimensions of the digitizer surface in device coordinates.
@@ -81,14 +81,14 @@ namespace UE::StylusInput
 		 *
 		 * @returns The dimensions of the digitizer surface.
 		 */
-		STYLUSINPUT_API virtual FIntRect GetInputRectangle() const = 0;
+		virtual FIntRect GetInputRectangle() const = 0;
 
 		/**
 		 * Returns the hardware capabilities of the device associated with the tablet context.
 		 *
 		 * @return The hardware capability flags.
 		 */
-		STYLUSINPUT_API virtual ETabletHardwareCapabilities GetHardwareCapabilities() const = 0;
+		virtual ETabletHardwareCapabilities GetHardwareCapabilities() const = 0;
 
 		/**
 		 * Returns the supported properties for a tablet context.
@@ -96,7 +96,7 @@ namespace UE::StylusInput
 		 *
 		 * @returns The supported properties flags.
 		 */
-		STYLUSINPUT_API virtual ETabletSupportedProperties GetSupportedProperties() const = 0;
+		virtual ETabletSupportedProperties GetSupportedProperties() const = 0;
 	};
 
 	/**
@@ -110,12 +110,12 @@ namespace UE::StylusInput
 		/**
 		 * @returns The GUID string for the button. 
 		 */
-		STYLUSINPUT_API virtual FString GetID() const = 0;
+		virtual FString GetID() const = 0;
 
 		/**
 		 * @returns  The name of the button.
 		 */
-		STYLUSINPUT_API virtual FString GetName() const = 0;
+		virtual FString GetName() const = 0;
 	};
 
 	/**
@@ -132,12 +132,12 @@ namespace UE::StylusInput
 		 *
 		 * @returns Identifier of the stylus pen.
 		 */
-		STYLUSINPUT_API virtual uint32 GetID() const = 0;
+		virtual uint32 GetID() const = 0;
 
 		/**
 		 * @returns The name of the stylus pen.
 		 */
-		STYLUSINPUT_API virtual FString GetName() const = 0;
+		virtual FString GetName() const = 0;
 
 		/**
 		 * Returns the number of buttons on the stylus.
@@ -145,7 +145,7 @@ namespace UE::StylusInput
 		 *
 		 * @returns The number of buttons on the stylus.
 		 */
-		STYLUSINPUT_API virtual uint32 GetNumButtons() const = 0;
+		virtual uint32 GetNumButtons() const = 0;
 
 		/**
 		 * Returns a description of the button for a given index.
@@ -153,7 +153,7 @@ namespace UE::StylusInput
 		 * @param Index A number in [0, @see #GetNumButtons).
 		 * @returns Pointer to a stylus button.
 		 */
-		STYLUSINPUT_API virtual const IStylusInputStylusButton* GetButton(int32 Index) const = 0;
+		virtual const IStylusInputStylusButton* GetButton(int32 Index) const = 0;
 	};
 
 	/* Bitwise operators for enums to be used as flags without additional casts. */
