StylusInput Windows:
* improve debug message output,
* clean up plugin overrides,
* move common plugin code into base class

diff --git a/Engine/Plugins/Editor/StylusInput/Source/StylusInput/Private/Windows/WindowsStylusInputInstance.cpp b/Engine/Plugins/Editor/StylusInput/Source/StylusInput/Private/Windows/WindowsStylusInputInstance.cpp
--- a/Engine/Plugins/Editor/StylusInput/Source/StylusInput/Private/Windows/WindowsStylusInputInstance.cpp
+++ b/Engine/Plugins/Editor/StylusInput/Source/StylusInput/Private/Windows/WindowsStylusInputInstance.cpp
@@ -46,15 +46,17 @@ namespace UE::StylusInput::Private::Windows
 		{
 			if (!AsyncPlugin)
 			{
-				EnablePlugin(EEventHandlerThread::OnGameThread);
-			}
+				EnablePlugin(EEventHandlerThread::OnGameThread, EventHandler);
 
-			if (!AsyncPlugin)
-			{
-				LogError(LOG_PREAMBLE, FString::Format(
-					         TEXT("Event handler '{0}' was not added since game thread stylus input plugin could not be installed."),
-					         {EventHandler->GetName()}));
-				return false;
+				if (!AsyncPlugin)
+				{
+					LogError(LOG_PREAMBLE, FString::Format(
+								 TEXT("Event handler '{0}' was not added since the game thread stylus input plugin could not be installed."),
+								 {EventHandler->GetName()}));
+					return false;
+				}
+
+				return true;
 			}
 
 			return AsyncPlugin->AddEventHandler(EventHandler);
@@ -64,15 +66,17 @@ namespace UE::StylusInput::Private::Windows
 		{
 			if (!SyncPlugin)
 			{
-				EnablePlugin(EEventHandlerThread::Asynchronous);
-			}
+				EnablePlugin(EEventHandlerThread::Asynchronous, EventHandler);
 
-			if (!SyncPlugin)
-			{
-				LogError(LOG_PREAMBLE, FString::Format(
-							 TEXT("Event handler '{0}' was not added since asynchronous stylus input plugin could not be installed."),
-							 {EventHandler->GetName()}));
-				return false;
+				if (!SyncPlugin)
+				{
+					LogError(LOG_PREAMBLE, FString::Format(
+								 TEXT("Event handler '{0}' was not added since the asynchronous stylus input plugin could not be installed."),
+								 {EventHandler->GetName()}));
+					return false;
+				}
+
+				return true;
 			}
 
 			return SyncPlugin->AddEventHandler(EventHandler);
@@ -250,7 +254,7 @@ namespace UE::StylusInput::Private::Windows
 
 		if (WindowRectangle.right <= WindowRectangle.left || WindowRectangle.bottom <= WindowRectangle.top)
 		{
-			LogError(LOG_PREAMBLE, "Window is minimized; failed to setup device context!");
+			LogError(LOG_PREAMBLE, "Window appears to be minimized; failed to setup device context!");
 			return;
 		}
 
@@ -361,7 +365,7 @@ namespace UE::StylusInput::Private::Windows
 		}
 	}
 
-	void FWindowsStylusInputInstance::EnablePlugin(const EEventHandlerThread EventHandlerThread)
+	void FWindowsStylusInputInstance::EnablePlugin(const EEventHandlerThread EventHandlerThread, IStylusInputEventHandler* EventHandler)
 	{
 		check(RealTimeStylus.IsValid());
 
@@ -369,7 +373,8 @@ namespace UE::StylusInput::Private::Windows
 		{
 			AsyncPlugin = MakeUnique<FWindowsStylusInputPluginAsync>(
 				FGetWindowContextCallback::CreateRaw(this, &FWindowsStylusInputInstance::GetWindowContext),
-				FUpdateTabletContextsCallback::CreateRaw(this, &FWindowsStylusInputInstance::UpdateTabletContexts)
+				FUpdateTabletContextsCallback::CreateRaw(this, &FWindowsStylusInputInstance::UpdateTabletContexts),
+				EventHandler
 			);
 
 			if (Succeeded(RealTimeStylus->AddStylusAsyncPlugin(ASYNC_PLUGIN_INDEX, AsyncPlugin.Get()), LOG_PREAMBLE))
@@ -387,7 +392,8 @@ namespace UE::StylusInput::Private::Windows
 		{
 			SyncPlugin = MakeUnique<FWindowsStylusInputPluginSync>(
 				FGetWindowContextCallback::CreateRaw(this, &FWindowsStylusInputInstance::GetWindowContext),
-				FUpdateTabletContextsCallback::CreateRaw(this, &FWindowsStylusInputInstance::UpdateTabletContexts)
+				FUpdateTabletContextsCallback::CreateRaw(this, &FWindowsStylusInputInstance::UpdateTabletContexts),
+				EventHandler
 			);
 
 			if (Succeeded(SyncPlugin->CreateFreeThreadMarshaler(), LOG_PREAMBLE))
diff --git a/Engine/Plugins/Editor/StylusInput/Source/StylusInput/Private/Windows/WindowsStylusInputInstance.h b/Engine/Plugins/Editor/StylusInput/Source/StylusInput/Private/Windows/WindowsStylusInputInstance.h
--- a/Engine/Plugins/Editor/StylusInput/Source/StylusInput/Private/Windows/WindowsStylusInputInstance.h
+++ b/Engine/Plugins/Editor/StylusInput/Source/StylusInput/Private/Windows/WindowsStylusInputInstance.h
@@ -45,7 +45,7 @@ namespace UE::StylusInput::Private::Windows
 	private:
 		void Init(HWND HWindow);
 
-		void EnablePlugin(EEventHandlerThread EventHandlerThread);
+		void EnablePlugin(EEventHandlerThread EventHandlerThread, IStylusInputEventHandler* EventHandler);
 		void DisablePlugin(EEventHandlerThread EventHandlerThread);
 
 		void SetupWindowContext(HWND HWindow);
diff --git a/Engine/Plugins/Editor/StylusInput/Source/StylusInput/Private/Windows/WindowsStylusInputPluginAsync.cpp b/Engine/Plugins/Editor/StylusInput/Source/StylusInput/Private/Windows/WindowsStylusInputPluginAsync.cpp
--- a/Engine/Plugins/Editor/StylusInput/Source/StylusInput/Private/Windows/WindowsStylusInputPluginAsync.cpp
+++ b/Engine/Plugins/Editor/StylusInput/Source/StylusInput/Private/Windows/WindowsStylusInputPluginAsync.cpp
@@ -7,8 +7,9 @@
 namespace UE::StylusInput::Private::Windows
 {
 	FWindowsStylusInputPluginAsync::FWindowsStylusInputPluginAsync(FGetWindowContextCallback&& GetWindowContext,
-	                                                     FUpdateTabletContextsCallback&& UpdateTabletContextsCallback)
-		: FWindowsStylusInputPluginBase(MoveTemp(GetWindowContext), MoveTemp(UpdateTabletContextsCallback))
+	                                                               FUpdateTabletContextsCallback&& UpdateTabletContextsCallback,
+	                                                               IStylusInputEventHandler* EventHandler)
+		: FWindowsStylusInputPluginBase(MoveTemp(GetWindowContext), MoveTemp(UpdateTabletContextsCallback), EventHandler)
 	{
 	}
 
@@ -16,22 +17,26 @@ namespace UE::StylusInput::Private::Windows
 	{
 	}
 
-	HRESULT FWindowsStylusInputPluginAsync::RealTimeStylusEnabled(IRealTimeStylus* RealTimeStylus, const ULONG TabletContextIDsCount, const TABLET_CONTEXT_ID* TabletContextIDs)
+	HRESULT FWindowsStylusInputPluginAsync::RealTimeStylusEnabled(IRealTimeStylus* RealTimeStylus, const ULONG TabletContextIDsCount,
+	                                                              const TABLET_CONTEXT_ID* TabletContextIDs)
 	{
 		return ProcessRealTimeStylusEnabled(RealTimeStylus, TabletContextIDsCount, TabletContextIDs);
 	}
 
-	HRESULT FWindowsStylusInputPluginAsync::RealTimeStylusDisabled(IRealTimeStylus* RealTimeStylus, const ULONG TabletContextIDsCount, const TABLET_CONTEXT_ID* TabletContextIDs)
+	HRESULT FWindowsStylusInputPluginAsync::RealTimeStylusDisabled(IRealTimeStylus* RealTimeStylus, const ULONG TabletContextIDsCount,
+	                                                               const TABLET_CONTEXT_ID* TabletContextIDs)
 	{
 		return ProcessRealTimeStylusDisabled(RealTimeStylus, TabletContextIDsCount, TabletContextIDs);
 	}
 
-	HRESULT FWindowsStylusInputPluginAsync::StylusDown(IRealTimeStylus* RealTimeStylus, const StylusInfo* StylusInfo, const ULONG PropertyCount, LONG* PacketBuffer, LONG**)
+	HRESULT FWindowsStylusInputPluginAsync::StylusDown(IRealTimeStylus* RealTimeStylus, const StylusInfo* StylusInfo, const ULONG PropertyCount,
+	                                                   LONG* PacketBuffer, LONG**)
 	{
 		return ProcessPackets(StylusInfo, 1, PropertyCount, EPacketType::StylusDown, reinterpret_cast<int32*>(PacketBuffer));
 	}
 
-	HRESULT FWindowsStylusInputPluginAsync::StylusUp(IRealTimeStylus* RealTimeStylus, const StylusInfo* StylusInfo, const ULONG PropertyCount, LONG* PacketBuffer, LONG**)
+	HRESULT FWindowsStylusInputPluginAsync::StylusUp(IRealTimeStylus* RealTimeStylus, const StylusInfo* StylusInfo, const ULONG PropertyCount,
+	                                                 LONG* PacketBuffer, LONG**)
 	{
 		return ProcessPackets(StylusInfo, 1, PropertyCount, EPacketType::StylusUp, reinterpret_cast<int32*>(PacketBuffer));
 	}
@@ -41,7 +46,7 @@ namespace UE::StylusInput::Private::Windows
 		return ProcessTabletAdded(Tablet);
 	}
 
-	HRESULT FWindowsStylusInputPluginAsync::TabletRemoved(IRealTimeStylus* RealTimeStylus, LONG TabletIndex)
+	HRESULT FWindowsStylusInputPluginAsync::TabletRemoved(IRealTimeStylus* RealTimeStylus, const LONG TabletIndex)
 	{
 		return ProcessTabletRemoved(TabletIndex);
 	}
@@ -68,45 +73,43 @@ namespace UE::StylusInput::Private::Windows
 		return E_NOTIMPL;
 	}
 
-	HRESULT FWindowsStylusInputPluginAsync::StylusButtonDown(IRealTimeStylus* RealTimeStylus, STYLUS_ID StylusID, const GUID* pGuidStylusButton, POINT* pStylusPos)
-	{		
+	HRESULT FWindowsStylusInputPluginAsync::StylusButtonDown(IRealTimeStylus* RealTimeStylus, STYLUS_ID StylusID, const GUID* GuidStylusButton,
+	                                                         POINT* StylusPos)
+	{
 		return E_NOTIMPL;
 	}
 
-	HRESULT FWindowsStylusInputPluginAsync::StylusButtonUp(IRealTimeStylus* RealTimeStylus, STYLUS_ID StylusID, const GUID* pGuidStylusButton, POINT* pStylusPos)
+	HRESULT FWindowsStylusInputPluginAsync::StylusButtonUp(IRealTimeStylus* RealTimeStylus, STYLUS_ID StylusID, const GUID* GuidStylusButton, POINT* StylusPos)
 	{
 		return E_NOTIMPL;
 	}
 
-	HRESULT FWindowsStylusInputPluginAsync::CustomStylusDataAdded(IRealTimeStylus* piRtsSrc, const GUID* pGuidId, ULONG cbData, const BYTE* pbData)
+	HRESULT FWindowsStylusInputPluginAsync::CustomStylusDataAdded(IRealTimeStylus* RealTimeStylus, const GUID* GuidId, ULONG DataCount, const BYTE* Data)
 	{
 		return E_NOTIMPL;
 	}
 
-	HRESULT FWindowsStylusInputPluginAsync::SystemEvent(IRealTimeStylus* piRtsSrc, TABLET_CONTEXT_ID tcid, STYLUS_ID sid, SYSTEM_EVENT event, SYSTEM_EVENT_DATA eventdata)
+	HRESULT FWindowsStylusInputPluginAsync::SystemEvent(IRealTimeStylus* RealTimeStylus, TABLET_CONTEXT_ID TabletContextID, STYLUS_ID StylusID,
+	                                                    SYSTEM_EVENT Event, SYSTEM_EVENT_DATA EventData)
 	{
 		return E_NOTIMPL;
 	}
 
-	HRESULT FWindowsStylusInputPluginAsync::Error(IRealTimeStylus*, IStylusPlugin* Plugin, const RealTimeStylusDataInterest DataInterest, const HRESULT ErrorCode, LONG_PTR*)
+	HRESULT FWindowsStylusInputPluginAsync::Error(IRealTimeStylus*, IStylusPlugin*, const RealTimeStylusDataInterest DataInterest, const HRESULT ErrorCode,
+	                                              LONG_PTR*)
 	{
-#if STYLUSINPUT_SHOW_NOTIMPL_ERRORS
-		const FString ErrorMessage = ProcessError(DataInterest, ErrorCode);
-		DebugEvent(ErrorMessage);
-#endif
-
-		return S_OK;
+		return ProcessError(DataInterest, ErrorCode);
 	}
 
-	HRESULT FWindowsStylusInputPluginAsync::UpdateMapping(IRealTimeStylus* piRtsSrc)
+	HRESULT FWindowsStylusInputPluginAsync::UpdateMapping(IRealTimeStylus* RealTimeStylus)
 	{
 		return E_NOTIMPL;
 	}
 
-	HRESULT FWindowsStylusInputPluginAsync::DataInterest(RealTimeStylusDataInterest* pDataInterest)
+	HRESULT FWindowsStylusInputPluginAsync::DataInterest(RealTimeStylusDataInterest* DataInterest)
 	{
-		*pDataInterest = RTSDI_AllData;
-		return S_OK;
+		return ProcessDataInterest(DataInterest);
 	}
 }
+
 #endif
diff --git a/Engine/Plugins/Editor/StylusInput/Source/StylusInput/Private/Windows/WindowsStylusInputPluginAsync.h b/Engine/Plugins/Editor/StylusInput/Source/StylusInput/Private/Windows/WindowsStylusInputPluginAsync.h
--- a/Engine/Plugins/Editor/StylusInput/Source/StylusInput/Private/Windows/WindowsStylusInputPluginAsync.h
+++ b/Engine/Plugins/Editor/StylusInput/Source/StylusInput/Private/Windows/WindowsStylusInputPluginAsync.h
@@ -13,7 +13,8 @@ namespace UE::StylusInput::Private::Windows
 	class FWindowsStylusInputPluginAsync final : public IStylusAsyncPlugin, public FWindowsStylusInputPluginBase
 	{
 	public:
-		explicit FWindowsStylusInputPluginAsync(FGetWindowContextCallback&& GetWindowContext, FUpdateTabletContextsCallback&& UpdateTabletContextsCallback);
+		FWindowsStylusInputPluginAsync(FGetWindowContextCallback&& GetWindowContext, FUpdateTabletContextsCallback&& UpdateTabletContextsCallback,
+		                               IStylusInputEventHandler* EventHandler);
 		~FWindowsStylusInputPluginAsync();
 
 		// IUnknown
@@ -47,21 +48,21 @@ namespace UE::StylusInput::Private::Windows
 		// IStylusPlugin
 		virtual HRESULT RealTimeStylusEnabled(IRealTimeStylus* RealTimeStylus, ULONG TabletContextIDsCount, const TABLET_CONTEXT_ID* TabletContextIDs) override;
 		virtual HRESULT RealTimeStylusDisabled(IRealTimeStylus* RealTimeStylus, ULONG TabletContextIDsCount, const TABLET_CONTEXT_ID* TabletContextIDs) override;
-		virtual HRESULT StylusDown(IRealTimeStylus* RealTimeStylus, const StylusInfo* StylusInfo, ULONG PropertyCount, LONG* PacketBuffer, LONG** ppInOutPkt) override;
-		virtual HRESULT StylusUp(IRealTimeStylus* RealTimeStylus, const StylusInfo* StylusInfo, ULONG PropertyCount, LONG* PacketBuffer, LONG** ppInOutPkt) override;
+		virtual HRESULT StylusDown(IRealTimeStylus* RealTimeStylus, const StylusInfo* StylusInfo, ULONG PropertyCount, LONG* PacketBuffer, LONG** InOutPkt) override;
+		virtual HRESULT StylusUp(IRealTimeStylus* RealTimeStylus, const StylusInfo* StylusInfo, ULONG PropertyCount, LONG* PacketBuffer, LONG** InOutPkt) override;
 		virtual HRESULT TabletAdded(IRealTimeStylus* RealTimeStylus, IInkTablet* Tablet) override;
 		virtual HRESULT TabletRemoved(IRealTimeStylus* RealTimeStylus, LONG TabletIndex) override;
 		virtual HRESULT InAirPackets(IRealTimeStylus* RealTimeStylus, const StylusInfo* StylusInfo, ULONG PacketCount, ULONG PacketBufferLength, LONG* PacketBuffer, ULONG* InOutPacketCount, LONG** InOutPacketBuffer) override;
 		virtual HRESULT Packets(IRealTimeStylus* RealTimeStylus, const StylusInfo* StylusInfo, ULONG PacketCount, ULONG PacketBufferLength, LONG* PacketBuffer, ULONG* InOutPacketCount, LONG** InOutPacketBuffer) override;
-		virtual HRESULT StylusInRange(IRealTimeStylus* piRtsSrc, TABLET_CONTEXT_ID tcid, STYLUS_ID sid) override;
-		virtual HRESULT StylusOutOfRange(IRealTimeStylus* piRtsSrc, TABLET_CONTEXT_ID tcid, STYLUS_ID sid) override;
-		virtual HRESULT StylusButtonDown(IRealTimeStylus* piRtsSrc, STYLUS_ID sid, const GUID* pGuidStylusButton, POINT* pStylusPos) override;
-		virtual HRESULT StylusButtonUp(IRealTimeStylus* piRtsSrc, STYLUS_ID sid, const GUID* pGuidStylusButton, POINT* pStylusPos) override;
-		virtual HRESULT CustomStylusDataAdded(IRealTimeStylus* piRtsSrc, const GUID* pGuidId, ULONG cbData, const BYTE* pbData) override;
-		virtual HRESULT SystemEvent(IRealTimeStylus* piRtsSrc, TABLET_CONTEXT_ID tcid, STYLUS_ID sid, SYSTEM_EVENT event, SYSTEM_EVENT_DATA eventdata) override;
+		virtual HRESULT StylusInRange(IRealTimeStylus* RealTimeStylus, TABLET_CONTEXT_ID TabletContextID, STYLUS_ID StylusID) override;
+		virtual HRESULT StylusOutOfRange(IRealTimeStylus* RealTimeStylus, TABLET_CONTEXT_ID TabletContextID, STYLUS_ID StylusID) override;
+		virtual HRESULT StylusButtonDown(IRealTimeStylus* RealTimeStylus, STYLUS_ID StylusID, const GUID* GuidStylusButton, POINT* StylusPos) override;
+		virtual HRESULT StylusButtonUp(IRealTimeStylus* RealTimeStylus, STYLUS_ID StylusID, const GUID* GuidStylusButton, POINT* StylusPos) override;
+		virtual HRESULT CustomStylusDataAdded(IRealTimeStylus* RealTimeStylus, const GUID* GuidId, ULONG DataCount, const BYTE* Data) override;
+		virtual HRESULT SystemEvent(IRealTimeStylus* RealTimeStylus, TABLET_CONTEXT_ID TabletContextID, STYLUS_ID StylusID, SYSTEM_EVENT Event, SYSTEM_EVENT_DATA EventData) override;
 		virtual HRESULT Error(IRealTimeStylus* RealTimeStylus, IStylusPlugin* Plugin, RealTimeStylusDataInterest DataInterest, HRESULT ErrorCode, LONG_PTR* InternalKey) override;
-		virtual HRESULT UpdateMapping(IRealTimeStylus* piRtsSrc) override;
-		virtual HRESULT DataInterest(RealTimeStylusDataInterest* pDataInterest) override;
+		virtual HRESULT UpdateMapping(IRealTimeStylus* RealTimeStylus) override;
+		virtual HRESULT DataInterest(RealTimeStylusDataInterest* DataInterest) override;
 
 	protected:
 		virtual FString GetName() const override { return "OnGameThread"; }
diff --git a/Engine/Plugins/Editor/StylusInput/Source/StylusInput/Private/Windows/WindowsStylusInputPluginBase.cpp b/Engine/Plugins/Editor/StylusInput/Source/StylusInput/Private/Windows/WindowsStylusInputPluginBase.cpp
--- a/Engine/Plugins/Editor/StylusInput/Source/StylusInput/Private/Windows/WindowsStylusInputPluginBase.cpp
+++ b/Engine/Plugins/Editor/StylusInput/Source/StylusInput/Private/Windows/WindowsStylusInputPluginBase.cpp
@@ -514,12 +514,19 @@ namespace UE::StylusInput::Private::Windows
 	}
 
 	FWindowsStylusInputPluginBase::FWindowsStylusInputPluginBase(FGetWindowContextCallback&& GetWindowContextCallback,
-	                                                   FUpdateTabletContextsCallback&& UpdateTabletContextsCallback)
+	                                                             FUpdateTabletContextsCallback&& UpdateTabletContextsCallback,
+	                                                             IStylusInputEventHandler* EventHandler)
 		: GetWindowContextCallback(MoveTemp(GetWindowContextCallback))
 		, UpdateTabletContextsCallback(MoveTemp(UpdateTabletContextsCallback)) 
 	{
 		check(this->GetWindowContextCallback.IsBound());
 		check(this->UpdateTabletContextsCallback.IsBound());
+
+		if (EventHandler)
+		{
+			// Immediately install an event handler during construction to capture events coming through during plugin initialization.
+			AddEventHandler(EventHandler);
+		}
 	}
 
 	void FWindowsStylusInputPluginBase::DebugEvent(const FString& Message) const
@@ -530,16 +537,93 @@ namespace UE::StylusInput::Private::Windows
 		}
 	}
 
-	HRESULT FWindowsStylusInputPluginBase::ProcessRealTimeStylusEnabled(IRealTimeStylus* RealTimeStylus, const uint32 TabletContextIDsCount,
-	                                                               const TABLET_CONTEXT_ID* TabletContextIDs)
+	HRESULT FWindowsStylusInputPluginBase::ProcessDataInterest(RealTimeStylusDataInterest* DataInterest)
 	{
-		return UpdateTabletContexts(RealTimeStylus, TabletContextIDsCount, TabletContextIDs);
+		constexpr bool bGetAllData = false;
+
+		if (bGetAllData)
+		{
+			*DataInterest = RTSDI_AllData;
+
+			DebugEvent("Requested all data from stylus input.");
+		}
+		else
+		{
+			*DataInterest = static_cast<RealTimeStylusDataInterest>(
+				RTSDI_Error |
+				RTSDI_RealTimeStylusEnabled | RTSDI_RealTimeStylusDisabled |
+				RTSDI_InAirPackets | RTSDI_Packets |
+				RTSDI_StylusDown | RTSDI_StylusUp);
+
+			DebugEvent("Requested data for Error, RealTimeStylusEnabled, RealTimeStylusDisabled, InAirPackets, Packets, StylusDown, StylusUp "
+				"from stylus input (DataInterest event).");
+		}
+
+		return S_OK;
 	}
 
-	HRESULT FWindowsStylusInputPluginBase::ProcessRealTimeStylusDisabled(IRealTimeStylus* RealTimeStylus, const uint32 TabletContextIDsCount,
-	                                                                const TABLET_CONTEXT_ID* TabletContextIDs)
+	HRESULT FWindowsStylusInputPluginBase::ProcessError(RealTimeStylusDataInterest DataInterest, const HRESULT ErrorCode)
 	{
-		return UpdateTabletContexts(RealTimeStylus, TabletContextIDsCount, TabletContextIDs);
+		constexpr bool bShowAllErrors = false;
+
+		if (bShowAllErrors || ErrorCode != E_NOTIMPL)
+		{
+			const char* DataInterestStr = [DataInterest]
+			{
+				switch (DataInterest)
+				{
+				case RTSDI_AllData:
+					return "AllData";
+				case RTSDI_None:
+					return "None";
+				case RTSDI_Error:
+					return "Error";
+				case RTSDI_RealTimeStylusEnabled:
+					return "RealTimeStylusEnabled";
+				case RTSDI_RealTimeStylusDisabled:
+					return "RealTimeStylusDisabled";
+				case RTSDI_StylusNew:
+					return "StylusNew";
+				case RTSDI_StylusInRange:
+					return "StylusInRange";
+				case RTSDI_InAirPackets:
+					return "InAirPackets";
+				case RTSDI_StylusOutOfRange:
+					return "StylusOutOfRange";
+				case RTSDI_StylusDown:
+					return "StylusDown";
+				case RTSDI_Packets:
+					return "Packets";
+				case RTSDI_StylusUp:
+					return "StylusUp";
+				case RTSDI_StylusButtonUp:
+					return "StylusButtonUp";
+				case RTSDI_StylusButtonDown:
+					return "StylusButtonDown";
+				case RTSDI_SystemEvents:
+					return "SystemEvents";
+				case RTSDI_TabletAdded:
+					return "TabletAdded";
+				case RTSDI_TabletRemoved:
+					return "TabletRemoved";
+				case RTSDI_CustomStylusDataAdded:
+					return "CustomStylusDataAdded";
+				case RTSDI_UpdateMapping:
+					return "UpdateMapping";
+				case RTSDI_DefaultEvents:
+					return "DefaultEvents";
+				}
+				return "<unknown>";
+			}();
+
+			const FString ErrorMessage = FString::Format(TEXT("Error in {0} plugin: {1}, Error={2} ({3})."), {
+				                                             GetName(), DataInterestStr, _com_error(ErrorCode).ErrorMessage(), static_cast<int32>(ErrorCode)
+			                                             });
+
+			DebugEvent(ErrorMessage);
+		}
+
+		return S_OK;
 	}
 
 	HRESULT FWindowsStylusInputPluginBase::ProcessPackets(const StylusInfo* StylusInfo, uint32 PacketCount, uint32 PacketBufferLength, EPacketType Type,
@@ -582,68 +666,36 @@ namespace UE::StylusInput::Private::Windows
 		return S_OK;
 	}
 
-	FString FWindowsStylusInputPluginBase::ProcessError(RealTimeStylusDataInterest DataInterest, const HRESULT ErrorCode)
+	HRESULT FWindowsStylusInputPluginBase::ProcessRealTimeStylusEnabled(IRealTimeStylus* RealTimeStylus, const uint32 TabletContextIDsCount,
+																   const TABLET_CONTEXT_ID* TabletContextIDs)
 	{
-		const char* DataInterestStr = [DataInterest]
-		{
-			switch (DataInterest)
-			{
-			case RTSDI_AllData:
-				return "AllData";
-			case RTSDI_None:
-				return "None";
-			case RTSDI_Error:
-				return "Error";
-			case RTSDI_RealTimeStylusEnabled:
-				return "RealTimeStylusEnabled";
-			case RTSDI_RealTimeStylusDisabled:
-				return "RealTimeStylusDisabled";
-			case RTSDI_StylusNew:
-				return "StylusNew";
-			case RTSDI_StylusInRange:
-				return "StylusInRange";
-			case RTSDI_InAirPackets:
-				return "InAirPackets";
-			case RTSDI_StylusOutOfRange:
-				return "StylusOutOfRange";
-			case RTSDI_StylusDown:
-				return "StylusDown";
-			case RTSDI_Packets:
-				return "Packets";
-			case RTSDI_StylusUp:
-				return "StylusUp";
-			case RTSDI_StylusButtonUp:
-				return "StylusButtonUp";
-			case RTSDI_StylusButtonDown:
-				return "StylusButtonDown";
-			case RTSDI_SystemEvents:
-				return "SystemEvents";
-			case RTSDI_TabletAdded:
-				return "TabletAdded";
-			case RTSDI_TabletRemoved:
-				return "TabletRemoved";
-			case RTSDI_CustomStylusDataAdded:
-				return "CustomStylusDataAdded";
-			case RTSDI_UpdateMapping:
-				return "UpdateMapping";
-			case RTSDI_DefaultEvents:
-				return "DefaultEvents";
-			}
-			return "<unknown>";
-		}();
+		DebugEvent("Stylus input was enabled (RealTimeStylusEnabled event).");
+
+		return UpdateTabletContexts(RealTimeStylus, TabletContextIDsCount, TabletContextIDs);
+	}
+
+	HRESULT FWindowsStylusInputPluginBase::ProcessRealTimeStylusDisabled(IRealTimeStylus* RealTimeStylus, const uint32 TabletContextIDsCount,
+																	const TABLET_CONTEXT_ID* TabletContextIDs)
+	{
+		DebugEvent("Stylus input was disabled (RealTimeStylusDisabled event).");
 
-		return FString::Format(TEXT("Error in {0} plugin: {1}, Error={2} ({3})"), {
-			                       GetName(), DataInterestStr, _com_error(ErrorCode).ErrorMessage(), static_cast<int32>(ErrorCode)
-		                       });
+		return UpdateTabletContexts(RealTimeStylus, TabletContextIDsCount, TabletContextIDs);
 	}
 
 	HRESULT FWindowsStylusInputPluginBase::ProcessTabletAdded(IInkTablet* Tablet)
 	{
+		// After a TabletAdded event, Windows Ink will fire a RealTimeStylusDisabled event directly followed by a RealTimeStylusEnabled event.
+		// We are using these two events instead to update the tablet contexts.
+
 		return E_NOTIMPL;
 	}
 
 	HRESULT FWindowsStylusInputPluginBase::ProcessTabletRemoved(LONG TabletIndex)
 	{
+		// For simplicity, we don't remove a tablet context when this event is received.
+		// However, sincere there are no more packets coming through for the removed tablet there should be nothing that's continuing to access the outdated
+		// tablet context data.
+
 		return E_NOTIMPL;
 	}
 
diff --git a/Engine/Plugins/Editor/StylusInput/Source/StylusInput/Private/Windows/WindowsStylusInputPluginBase.h b/Engine/Plugins/Editor/StylusInput/Source/StylusInput/Private/Windows/WindowsStylusInputPluginBase.h
--- a/Engine/Plugins/Editor/StylusInput/Source/StylusInput/Private/Windows/WindowsStylusInputPluginBase.h
+++ b/Engine/Plugins/Editor/StylusInput/Source/StylusInput/Private/Windows/WindowsStylusInputPluginBase.h
@@ -39,17 +39,19 @@ namespace UE::StylusInput::Private::Windows
 		float GetPacketsPerSecond() const { return PacketStats.GetPacketsPerSecond(); }
 
 	protected:
-		FWindowsStylusInputPluginBase(FGetWindowContextCallback&& GetWindowContextCallback, FUpdateTabletContextsCallback&& UpdateTabletContextsCallback);
+		FWindowsStylusInputPluginBase(FGetWindowContextCallback&& GetWindowContextCallback, FUpdateTabletContextsCallback&& UpdateTabletContextsCallback,
+		                              IStylusInputEventHandler* EventHandler);
 		~FWindowsStylusInputPluginBase() = default;
 
 		virtual FString GetName() const = 0;
 
 		void DebugEvent(const FString& Message) const;
 
+		HRESULT ProcessDataInterest(RealTimeStylusDataInterest* DataInterest);
+		HRESULT ProcessError(RealTimeStylusDataInterest DataInterest, HRESULT ErrorCode);
+		HRESULT ProcessPackets(const StylusInfo* StylusInfo, uint32 PacketCount, uint32 PacketBufferLength, EPacketType Type, const int32* PacketBuffer);
 		HRESULT ProcessRealTimeStylusEnabled(IRealTimeStylus* RealTimeStylus, uint32 TabletContextIDsCount, const TABLET_CONTEXT_ID* TabletContextIDs);
 		HRESULT ProcessRealTimeStylusDisabled(IRealTimeStylus* RealTimeStylus, uint32 TabletContextIDsCount, const TABLET_CONTEXT_ID* TabletContextIDs);
-		HRESULT ProcessPackets(const StylusInfo* StylusInfo, uint32 PacketCount, uint32 PacketBufferLength, EPacketType Type, const int32* PacketBuffer);
-		FString ProcessError(RealTimeStylusDataInterest DataInterest, HRESULT ErrorCode);
 		HRESULT ProcessTabletAdded(IInkTablet* Tablet);
 		HRESULT ProcessTabletRemoved(LONG TabletIndex);
 
@@ -68,6 +70,4 @@ namespace UE::StylusInput::Private::Windows
 	};
 }
 
-#define STYLUSINPUT_SHOW_NOTIMPL_ERRORS 0
-
 #endif
diff --git a/Engine/Plugins/Editor/StylusInput/Source/StylusInput/Private/Windows/WindowsStylusInputPluginSync.cpp b/Engine/Plugins/Editor/StylusInput/Source/StylusInput/Private/Windows/WindowsStylusInputPluginSync.cpp
--- a/Engine/Plugins/Editor/StylusInput/Source/StylusInput/Private/Windows/WindowsStylusInputPluginSync.cpp
+++ b/Engine/Plugins/Editor/StylusInput/Source/StylusInput/Private/Windows/WindowsStylusInputPluginSync.cpp
@@ -11,8 +11,9 @@
 namespace UE::StylusInput::Private::Windows
 {
 	FWindowsStylusInputPluginSync::FWindowsStylusInputPluginSync(FGetWindowContextCallback&& GetWindowContext,
-	                                                   FUpdateTabletContextsCallback&& UpdateTabletContextsCallback)
-		: FWindowsStylusInputPluginBase(MoveTemp(GetWindowContext), MoveTemp(UpdateTabletContextsCallback))
+	                                                             FUpdateTabletContextsCallback&& UpdateTabletContextsCallback,
+	                                                             IStylusInputEventHandler* EventHandler)
+		: FWindowsStylusInputPluginBase(MoveTemp(GetWindowContext), MoveTemp(UpdateTabletContextsCallback), EventHandler)
 	{
 	}
 
@@ -32,22 +33,26 @@ namespace UE::StylusInput::Private::Windows
 		return WindowsAPI.CoCreateFreeThreadedMarshaler(this, &FreeThreadedMarshaler);
 	}
 
-	HRESULT FWindowsStylusInputPluginSync::RealTimeStylusEnabled(IRealTimeStylus* RealTimeStylus, ULONG TabletContextIDsCount, const TABLET_CONTEXT_ID* TabletContextIDs)
+	HRESULT FWindowsStylusInputPluginSync::RealTimeStylusEnabled(IRealTimeStylus* RealTimeStylus, const ULONG TabletContextIDsCount,
+	                                                             const TABLET_CONTEXT_ID* TabletContextIDs)
 	{
 		return ProcessRealTimeStylusEnabled(RealTimeStylus, TabletContextIDsCount, TabletContextIDs);
 	}
 
-	HRESULT FWindowsStylusInputPluginSync::RealTimeStylusDisabled(IRealTimeStylus* RealTimeStylus, ULONG TabletContextIDsCount, const TABLET_CONTEXT_ID* TabletContextIDs)
+	HRESULT FWindowsStylusInputPluginSync::RealTimeStylusDisabled(IRealTimeStylus* RealTimeStylus, const ULONG TabletContextIDsCount,
+	                                                              const TABLET_CONTEXT_ID* TabletContextIDs)
 	{
 		return ProcessRealTimeStylusDisabled(RealTimeStylus, TabletContextIDsCount, TabletContextIDs);
 	}
 
-	HRESULT FWindowsStylusInputPluginSync::StylusDown(IRealTimeStylus* RealTimeStylus, const StylusInfo* StylusInfo, ULONG PropertyCount, LONG* PacketBuffer, LONG**)
+	HRESULT FWindowsStylusInputPluginSync::StylusDown(IRealTimeStylus* RealTimeStylus, const StylusInfo* StylusInfo, const ULONG PropertyCount,
+	                                                  LONG* PacketBuffer, LONG**)
 	{
 		return ProcessPackets(StylusInfo, 1, PropertyCount, EPacketType::StylusDown, reinterpret_cast<int32*>(PacketBuffer));
 	}
 
-	HRESULT FWindowsStylusInputPluginSync::StylusUp(IRealTimeStylus* RealTimeStylus, const StylusInfo* StylusInfo, ULONG PropertyCount, LONG* PacketBuffer, LONG**)
+	HRESULT FWindowsStylusInputPluginSync::StylusUp(IRealTimeStylus* RealTimeStylus, const StylusInfo* StylusInfo, const ULONG PropertyCount,
+	                                                LONG* PacketBuffer, LONG**)
 	{
 		return ProcessPackets(StylusInfo, 1, PropertyCount, EPacketType::StylusUp, reinterpret_cast<int32*>(PacketBuffer));
 	}
@@ -57,19 +62,19 @@ namespace UE::StylusInput::Private::Windows
 		return ProcessTabletAdded(Tablet);
 	}
 
-	HRESULT FWindowsStylusInputPluginSync::TabletRemoved(IRealTimeStylus* RealTimeStylus, LONG TabletIndex)
+	HRESULT FWindowsStylusInputPluginSync::TabletRemoved(IRealTimeStylus* RealTimeStylus, const LONG TabletIndex)
 	{
 		return ProcessTabletRemoved(TabletIndex);
 	}
 
-	HRESULT FWindowsStylusInputPluginSync::InAirPackets(IRealTimeStylus* RealTimeStylus, const StylusInfo* StylusInfo, ULONG PacketCount, ULONG PacketBufferLength,
-													LONG* PacketBuffer, ULONG*, LONG**)
+	HRESULT FWindowsStylusInputPluginSync::InAirPackets(IRealTimeStylus* RealTimeStylus, const StylusInfo* StylusInfo, const ULONG PacketCount,
+	                                                    const ULONG PacketBufferLength, LONG* PacketBuffer, ULONG*, LONG**)
 	{
 		return ProcessPackets(StylusInfo, PacketCount, PacketBufferLength, EPacketType::AboveDigitizer, reinterpret_cast<int32*>(PacketBuffer));
 	}
 
 	HRESULT FWindowsStylusInputPluginSync::Packets(IRealTimeStylus* RealTimeStylus, const StylusInfo* StylusInfo, const ULONG PacketCount,
-	                                          const ULONG PacketBufferLength, LONG* PacketBuffer, ULONG*, LONG**)
+	                                               const ULONG PacketBufferLength, LONG* PacketBuffer, ULONG*, LONG**)
 	{
 		return ProcessPackets(StylusInfo, PacketCount, PacketBufferLength, EPacketType::OnDigitizer, reinterpret_cast<int32*>(PacketBuffer));
 	}
@@ -84,45 +89,42 @@ namespace UE::StylusInput::Private::Windows
 		return E_NOTIMPL;
 	}
 
-	HRESULT FWindowsStylusInputPluginSync::StylusButtonDown(IRealTimeStylus* piRtsSrc, STYLUS_ID sid, const GUID* pGuidStylusButton, POINT* pStylusPos)
+	HRESULT FWindowsStylusInputPluginSync::StylusButtonDown(IRealTimeStylus* RealTimeStylus, STYLUS_ID StylusID, const GUID* GuidStylusButton, POINT* StylusPos)
 	{
 		return E_NOTIMPL;
 	}
 
-	HRESULT FWindowsStylusInputPluginSync::StylusButtonUp(IRealTimeStylus* piRtsSrc, STYLUS_ID sid, const GUID* pGuidStylusButton, POINT* pStylusPos)
+	HRESULT FWindowsStylusInputPluginSync::StylusButtonUp(IRealTimeStylus* RealTimeStylus, STYLUS_ID StylusID, const GUID* GuidStylusButton, POINT* StylusPos)
 	{
 		return E_NOTIMPL;
 	}
 
-	HRESULT FWindowsStylusInputPluginSync::CustomStylusDataAdded(IRealTimeStylus* piRtsSrc, const GUID* pGuidId, ULONG cbData, const BYTE* pbData)
+	HRESULT FWindowsStylusInputPluginSync::CustomStylusDataAdded(IRealTimeStylus* RealTimeStylus, const GUID* GuidId, ULONG DataCount, const BYTE* Data)
 	{
 		return E_NOTIMPL;
 	}
 
-	HRESULT FWindowsStylusInputPluginSync::SystemEvent(IRealTimeStylus* piRtsSrc, TABLET_CONTEXT_ID tcid, STYLUS_ID sid, SYSTEM_EVENT event, SYSTEM_EVENT_DATA eventdata)
+	HRESULT FWindowsStylusInputPluginSync::SystemEvent(IRealTimeStylus* RealTimeStylus, TABLET_CONTEXT_ID TabletContextID, STYLUS_ID StylusID,
+	                                                   SYSTEM_EVENT Event, SYSTEM_EVENT_DATA EventData)
 	{
 		return E_NOTIMPL;
 	}
 
-	HRESULT FWindowsStylusInputPluginSync::Error(IRealTimeStylus*, IStylusPlugin* Plugin, const RealTimeStylusDataInterest DataInterest, const HRESULT ErrorCode, LONG_PTR*)
+	HRESULT FWindowsStylusInputPluginSync::Error(IRealTimeStylus*, IStylusPlugin* Plugin, const RealTimeStylusDataInterest DataInterest,
+	                                             const HRESULT ErrorCode, LONG_PTR*)
 	{
-#if STYLUSINPUT_SHOW_NOTIMPL_ERRORS
-		const FString ErrorMessage = ProcessError(DataInterest, ErrorCode);
-		DebugEvent(ErrorMessage);
-#endif
-
-		return S_OK;
+		return ProcessError(DataInterest, ErrorCode);
 	}
 
-	HRESULT FWindowsStylusInputPluginSync::UpdateMapping(IRealTimeStylus* piRtsSrc)
+	HRESULT FWindowsStylusInputPluginSync::UpdateMapping(IRealTimeStylus* RealTimeStylus)
 	{
 		return E_NOTIMPL;
 	}
 
-	HRESULT FWindowsStylusInputPluginSync::DataInterest(RealTimeStylusDataInterest* pDataInterest)
+	HRESULT FWindowsStylusInputPluginSync::DataInterest(RealTimeStylusDataInterest* DataInterest)
 	{
-		*pDataInterest = RTSDI_AllData;
-		return S_OK;
+		return ProcessDataInterest(DataInterest);
 	}
 }
+
 #endif
diff --git a/Engine/Plugins/Editor/StylusInput/Source/StylusInput/Private/Windows/WindowsStylusInputPluginSync.h b/Engine/Plugins/Editor/StylusInput/Source/StylusInput/Private/Windows/WindowsStylusInputPluginSync.h
--- a/Engine/Plugins/Editor/StylusInput/Source/StylusInput/Private/Windows/WindowsStylusInputPluginSync.h
+++ b/Engine/Plugins/Editor/StylusInput/Source/StylusInput/Private/Windows/WindowsStylusInputPluginSync.h
@@ -13,7 +13,8 @@ namespace UE::StylusInput::Private::Windows
 	class FWindowsStylusInputPluginSync final : public IStylusSyncPlugin, public FWindowsStylusInputPluginBase
 	{
 	public:
-		explicit FWindowsStylusInputPluginSync(FGetWindowContextCallback&& GetWindowContext, FUpdateTabletContextsCallback&& UpdateTabletContextsCallback);
+		FWindowsStylusInputPluginSync(FGetWindowContextCallback&& GetWindowContext, FUpdateTabletContextsCallback&& UpdateTabletContextsCallback,
+		                              IStylusInputEventHandler* EventHandler);
 		~FWindowsStylusInputPluginSync();
 		
 		HRESULT CreateFreeThreadMarshaler();
@@ -54,21 +55,21 @@ namespace UE::StylusInput::Private::Windows
 		// IStylusPlugin
 		virtual HRESULT RealTimeStylusEnabled(IRealTimeStylus* RealTimeStylus, ULONG TabletContextIDsCount, const TABLET_CONTEXT_ID* TabletContextIDs) override;
 		virtual HRESULT RealTimeStylusDisabled(IRealTimeStylus* RealTimeStylus, ULONG TabletContextIDsCount, const TABLET_CONTEXT_ID* TabletContextIDs) override;
-		virtual HRESULT StylusDown(IRealTimeStylus* RealTimeStylus, const StylusInfo* StylusInfo, ULONG PropertyCount, LONG* PacketBuffer, LONG** ppInOutPkt) override;
-		virtual HRESULT StylusUp(IRealTimeStylus* piRtsSrc, const StylusInfo* pStylusInfo, ULONG cPropCountPerPkt, LONG* pPacket, LONG** ppInOutPkt) override;
+		virtual HRESULT StylusDown(IRealTimeStylus* RealTimeStylus, const StylusInfo* StylusInfo, ULONG PropertyCount, LONG* PacketBuffer, LONG** InOutPkt) override;
+		virtual HRESULT StylusUp(IRealTimeStylus* RealTimeStylus, const StylusInfo* StylusInfo, ULONG PropertyCount, LONG* PacketBuffer, LONG** InOutPkt) override;
 		virtual HRESULT TabletAdded(IRealTimeStylus* RealTimeStylus, IInkTablet* Tablet) override;
 		virtual HRESULT TabletRemoved(IRealTimeStylus* RealTimeStylus, LONG TabletIndex) override;
 		virtual HRESULT InAirPackets(IRealTimeStylus* RealTimeStylus, const StylusInfo* StylusInfo, ULONG PacketCount, ULONG PacketBufferLength, LONG* PacketBuffer, ULONG* InOutPacketCount, LONG** InOutPacketBuffer) override;
 		virtual HRESULT Packets(IRealTimeStylus* RealTimeStylus, const StylusInfo* StylusInfo, ULONG PacketCount, ULONG PacketBufferLength, LONG* PacketBuffer, ULONG* InOutPacketCount, LONG** InOutPacketBuffer) override;
-		virtual HRESULT StylusInRange(IRealTimeStylus* piRtsSrc, TABLET_CONTEXT_ID tcid, STYLUS_ID sid) override;
-		virtual HRESULT StylusOutOfRange(IRealTimeStylus* piRtsSrc, TABLET_CONTEXT_ID tcid, STYLUS_ID sid) override;
-		virtual HRESULT StylusButtonDown(IRealTimeStylus* piRtsSrc, STYLUS_ID sid, const GUID* pGuidStylusButton, POINT* pStylusPos) override;
-		virtual HRESULT StylusButtonUp(IRealTimeStylus* piRtsSrc, STYLUS_ID sid, const GUID* pGuidStylusButton, POINT* pStylusPos) override;
-		virtual HRESULT CustomStylusDataAdded(IRealTimeStylus* piRtsSrc, const GUID* pGuidId, ULONG cbData, const BYTE* pbData) override;
-		virtual HRESULT SystemEvent(IRealTimeStylus* piRtsSrc, TABLET_CONTEXT_ID tcid, STYLUS_ID sid, SYSTEM_EVENT event, SYSTEM_EVENT_DATA eventdata) override;
+		virtual HRESULT StylusInRange(IRealTimeStylus* RealTimeStylus, TABLET_CONTEXT_ID TabletContextID, STYLUS_ID StylusID) override;
+		virtual HRESULT StylusOutOfRange(IRealTimeStylus* RealTimeStylus, TABLET_CONTEXT_ID TabletContextID, STYLUS_ID StylusID) override;
+		virtual HRESULT StylusButtonDown(IRealTimeStylus* RealTimeStylus, STYLUS_ID StylusID, const GUID* GuidStylusButton, POINT* StylusPos) override;
+		virtual HRESULT StylusButtonUp(IRealTimeStylus* RealTimeStylus, STYLUS_ID StylusID, const GUID* GuidStylusButton, POINT* StylusPos) override;
+		virtual HRESULT CustomStylusDataAdded(IRealTimeStylus* RealTimeStylus, const GUID* GuidId, ULONG DataCount, const BYTE* Data) override;
+		virtual HRESULT SystemEvent(IRealTimeStylus* RealTimeStylus, TABLET_CONTEXT_ID TabletContextID, STYLUS_ID StylusID, SYSTEM_EVENT Event, SYSTEM_EVENT_DATA EventData) override;
 		virtual HRESULT Error(IRealTimeStylus* RealTimeStylus, IStylusPlugin* Plugin, RealTimeStylusDataInterest DataInterest, HRESULT ErrorCode, LONG_PTR* InternalKey) override;
-		virtual HRESULT UpdateMapping(IRealTimeStylus* piRtsSrc) override;
-		virtual HRESULT DataInterest(RealTimeStylusDataInterest* pDataInterest) override;
+		virtual HRESULT UpdateMapping(IRealTimeStylus* RealTimeStylus) override;
+		virtual HRESULT DataInterest(RealTimeStylusDataInterest* DataInterest) override;
 
 	protected:
 		// FWindowsStylusInputPluginBase
