StylusInput: Windows platform: Fix issue with calling virtual function in base class constructor

diff --git a/Engine/Plugins/Editor/StylusInput/Source/StylusInput/Private/Windows/WindowsStylusInputPluginAsync.cpp b/Engine/Plugins/Editor/StylusInput/Source/StylusInput/Private/Windows/WindowsStylusInputPluginAsync.cpp
--- a/Engine/Plugins/Editor/StylusInput/Source/StylusInput/Private/Windows/WindowsStylusInputPluginAsync.cpp
+++ b/Engine/Plugins/Editor/StylusInput/Source/StylusInput/Private/Windows/WindowsStylusInputPluginAsync.cpp
@@ -9,8 +9,13 @@ namespace UE::StylusInput::Private::Windows
 	FWindowsStylusInputPluginAsync::FWindowsStylusInputPluginAsync(FGetWindowContextCallback&& GetWindowContext,
 	                                                               FUpdateTabletContextsCallback&& UpdateTabletContextsCallback,
 	                                                               IStylusInputEventHandler* EventHandler)
-		: FWindowsStylusInputPluginBase(MoveTemp(GetWindowContext), MoveTemp(UpdateTabletContextsCallback), EventHandler)
+		: FWindowsStylusInputPluginBase(MoveTemp(GetWindowContext), MoveTemp(UpdateTabletContextsCallback))
 	{
+		if (EventHandler)
+		{
+			// Immediately install an event handler during construction to capture events coming through during plugin initialization.
+			AddEventHandler(EventHandler);
+		}
 	}
 
 	FWindowsStylusInputPluginAsync::~FWindowsStylusInputPluginAsync()
diff --git a/Engine/Plugins/Editor/StylusInput/Source/StylusInput/Private/Windows/WindowsStylusInputPluginBase.cpp b/Engine/Plugins/Editor/StylusInput/Source/StylusInput/Private/Windows/WindowsStylusInputPluginBase.cpp
--- a/Engine/Plugins/Editor/StylusInput/Source/StylusInput/Private/Windows/WindowsStylusInputPluginBase.cpp
+++ b/Engine/Plugins/Editor/StylusInput/Source/StylusInput/Private/Windows/WindowsStylusInputPluginBase.cpp
@@ -514,19 +514,12 @@ namespace UE::StylusInput::Private::Windows
 	}
 
 	FWindowsStylusInputPluginBase::FWindowsStylusInputPluginBase(FGetWindowContextCallback&& GetWindowContextCallback,
-	                                                             FUpdateTabletContextsCallback&& UpdateTabletContextsCallback,
-	                                                             IStylusInputEventHandler* EventHandler)
+	                                                             FUpdateTabletContextsCallback&& UpdateTabletContextsCallback)
 		: GetWindowContextCallback(MoveTemp(GetWindowContextCallback))
 		, UpdateTabletContextsCallback(MoveTemp(UpdateTabletContextsCallback)) 
 	{
 		check(this->GetWindowContextCallback.IsBound());
 		check(this->UpdateTabletContextsCallback.IsBound());
-
-		if (EventHandler)
-		{
-			// Immediately install an event handler during construction to capture events coming through during plugin initialization.
-			AddEventHandler(EventHandler);
-		}
 	}
 
 	void FWindowsStylusInputPluginBase::DebugEvent(const FString& Message) const
diff --git a/Engine/Plugins/Editor/StylusInput/Source/StylusInput/Private/Windows/WindowsStylusInputPluginBase.h b/Engine/Plugins/Editor/StylusInput/Source/StylusInput/Private/Windows/WindowsStylusInputPluginBase.h
--- a/Engine/Plugins/Editor/StylusInput/Source/StylusInput/Private/Windows/WindowsStylusInputPluginBase.h
+++ b/Engine/Plugins/Editor/StylusInput/Source/StylusInput/Private/Windows/WindowsStylusInputPluginBase.h
@@ -39,8 +39,7 @@ namespace UE::StylusInput::Private::Windows
 		float GetPacketsPerSecond() const { return PacketStats.GetPacketsPerSecond(); }
 
 	protected:
-		FWindowsStylusInputPluginBase(FGetWindowContextCallback&& GetWindowContextCallback, FUpdateTabletContextsCallback&& UpdateTabletContextsCallback,
-		                              IStylusInputEventHandler* EventHandler);
+		FWindowsStylusInputPluginBase(FGetWindowContextCallback&& GetWindowContextCallback, FUpdateTabletContextsCallback&& UpdateTabletContextsCallback);
 		~FWindowsStylusInputPluginBase() = default;
 
 		virtual FString GetName() const = 0;
diff --git a/Engine/Plugins/Editor/StylusInput/Source/StylusInput/Private/Windows/WindowsStylusInputPluginSync.cpp b/Engine/Plugins/Editor/StylusInput/Source/StylusInput/Private/Windows/WindowsStylusInputPluginSync.cpp
--- a/Engine/Plugins/Editor/StylusInput/Source/StylusInput/Private/Windows/WindowsStylusInputPluginSync.cpp
+++ b/Engine/Plugins/Editor/StylusInput/Source/StylusInput/Private/Windows/WindowsStylusInputPluginSync.cpp
@@ -13,8 +13,13 @@ namespace UE::StylusInput::Private::Windows
 	FWindowsStylusInputPluginSync::FWindowsStylusInputPluginSync(FGetWindowContextCallback&& GetWindowContext,
 	                                                             FUpdateTabletContextsCallback&& UpdateTabletContextsCallback,
 	                                                             IStylusInputEventHandler* EventHandler)
-		: FWindowsStylusInputPluginBase(MoveTemp(GetWindowContext), MoveTemp(UpdateTabletContextsCallback), EventHandler)
+		: FWindowsStylusInputPluginBase(MoveTemp(GetWindowContext), MoveTemp(UpdateTabletContextsCallback))
 	{
+		if (EventHandler)
+		{
+			// Immediately install an event handler during construction to capture events coming through during plugin initialization.
+			AddEventHandler(EventHandler);
+		}
 	}
 
 	FWindowsStylusInputPluginSync::~FWindowsStylusInputPluginSync()
