StylusInput: Fix debug widget drawing issue

diff --git a/Engine/Plugins/Editor/StylusInput/Source/StylusInputDebugWidget/Private/StylusInputDebugPaintWidget.cpp b/Engine/Plugins/Editor/StylusInput/Source/StylusInputDebugWidget/Private/StylusInputDebugPaintWidget.cpp
--- a/Engine/Plugins/Editor/StylusInput/Source/StylusInputDebugWidget/Private/StylusInputDebugPaintWidget.cpp
+++ b/Engine/Plugins/Editor/StylusInput/Source/StylusInputDebugWidget/Private/StylusInputDebugPaintWidget.cpp
@@ -81,23 +81,23 @@ namespace UE::StylusInput::DebugWidget
 
 		const float CurrentTimeMS = FPlatformTime::ToMilliseconds(FPlatformTime::Cycles());
 
-		int32 NumToRemove = 0;
-		for (const FPaintPacket& PaintPacket : PaintPackets)
+		auto GetOpacity = [FadeOutTimeMSPenDown, FadeOutTimeMSPenUp, CurrentTimeMS](const FPaintPacket& PaintPacket)
 		{
 			const bool bPenIsDown = (PaintPacket.PenStatus & EPenStatus::CursorIsTouching) != EPenStatus::None;
-			const bool bPenIsInverted = (PaintPacket.PenStatus & EPenStatus::CursorIsInverted) != EPenStatus::None;
-
 			const float DeltaTimeMS = CurrentTimeMS - PaintPacket.TimeAddedMS;
 			const float FadeOutTimeMS = bPenIsDown ? FadeOutTimeMSPenDown : FadeOutTimeMSPenUp;
+			return (FadeOutTimeMS - DeltaTimeMS) / FadeOutTimeMS;
+		};
+		
+		for (const FPaintPacket& PaintPacket : PaintPackets)
+		{
+			const bool bPenIsDown = (PaintPacket.PenStatus & EPenStatus::CursorIsTouching) != EPenStatus::None;
+			const bool bPenIsInverted = (PaintPacket.PenStatus & EPenStatus::CursorIsInverted) != EPenStatus::None;
 
-			if (DeltaTimeMS >= FadeOutTimeMS)
-			{
-				++NumToRemove;
-			}
-			else
+			if (const float Opacity = GetOpacity(PaintPacket); Opacity > 0.0f)
 			{
 				FLinearColor Color = bPenIsDown ? (bPenIsInverted ? ColorPenDownInverted : ColorPenDown) : ColorPenUp;
-				Color.A = (FadeOutTimeMS - DeltaTimeMS) / FadeOutTimeMS;
+				Color.A = Opacity;
 
 				const float Size = FMath::Max(MinSize, PaintPacket.NormalPressure * MaxSize);
 
@@ -107,11 +107,7 @@ namespace UE::StylusInput::DebugWidget
 			}
 		}
 
-		if (NumToRemove)
-		{
-			checkSlow(NumToRemove <= PaintPackets.Num());
-			PaintPackets.PopFrontNoCheck(NumToRemove);
-		}
+		PaintPackets.RemoveAll([&GetOpacity](const FPaintPacket& PaintPacket) { return GetOpacity(PaintPacket) <= 0.0f; });
 
 		return ++LayerId;
 	}
