StylusInput:
* Fix debug drawing showing outdated packets
* Properly clip debug widget drawing
* Simplify debug widget overlay setup

diff --git a/Engine/Plugins/Editor/StylusInput/Source/StylusInputDebugWidget/Private/StylusInputDebugPaintWidget.cpp b/Engine/Plugins/Editor/StylusInput/Source/StylusInputDebugWidget/Private/StylusInputDebugPaintWidget.cpp
--- a/Engine/Plugins/Editor/StylusInput/Source/StylusInputDebugWidget/Private/StylusInputDebugPaintWidget.cpp
+++ b/Engine/Plugins/Editor/StylusInput/Source/StylusInputDebugWidget/Private/StylusInputDebugPaintWidget.cpp
@@ -36,7 +36,7 @@ namespace UE::StylusInput::DebugWidget
 	void SStylusInputDebugPaintWidget::Add(const FStylusInputPacket& Packet)
 	{
 		PaintPackets.Emplace(FPaintPacket{
-			FVector2f{Packet.X, Packet.Y}, Packet.NormalPressure, FPlatformTime::ToMilliseconds(FPlatformTime::Cycles()), Packet.PenStatus
+			FVector2f{Packet.X, Packet.Y}, Packet.NormalPressure, Packet.PenStatus, FPlatformTime::ToMilliseconds64(FPlatformTime::Cycles64())
 		});
 	}
 
@@ -59,8 +59,8 @@ namespace UE::StylusInput::DebugWidget
 
 		constexpr float MinSize = 2.0f;
 		constexpr float MaxSize = 15.0f;
-		constexpr float FadeOutTimeMSPenDown = 3000.0f;
-		constexpr float FadeOutTimeMSPenUp = 1000.0f;
+		constexpr double FadeOutTimeMSPenDown = 3000.0;
+		constexpr double FadeOutTimeMSPenUp = 1000.0;
 
 		constexpr FLinearColor ColorPenDown = {0.0f, 0.2f, 1.0f};
 		constexpr FLinearColor ColorPenDownInverted = {1.0f, 0.1f, 0.0f};
@@ -79,33 +79,37 @@ namespace UE::StylusInput::DebugWidget
 		const Slate::FDeprecateVector2DResult PaintWidgetPosition = GetPaintSpaceGeometry().GetAbsolutePosition();
 		const FDeprecateSlateVector2D PositionOffset = PaintWidgetPosition - WindowPosition;
 
-		const float CurrentTimeMS = FPlatformTime::ToMilliseconds(FPlatformTime::Cycles());
+		const double CurrentTimeMS = FPlatformTime::ToMilliseconds64(FPlatformTime::Cycles64());
 
-		auto GetOpacity = [FadeOutTimeMSPenDown, FadeOutTimeMSPenUp, CurrentTimeMS](const FPaintPacket& PaintPacket)
+		auto GetOpacity = [FadeOutTimeMSPenDown, FadeOutTimeMSPenUp, CurrentTimeMS](const FPaintPacket& PaintPacket) -> float
 		{
 			const bool bPenIsDown = (PaintPacket.PenStatus & EPenStatus::CursorIsTouching) != EPenStatus::None;
-			const float DeltaTimeMS = CurrentTimeMS - PaintPacket.TimeAddedMS;
-			const float FadeOutTimeMS = bPenIsDown ? FadeOutTimeMSPenDown : FadeOutTimeMSPenUp;
-			return (FadeOutTimeMS - DeltaTimeMS) / FadeOutTimeMS;
+			const double DeltaTimeMS = CurrentTimeMS - PaintPacket.TimeAddedMS;
+			const double FadeOutTimeMS = bPenIsDown ? FadeOutTimeMSPenDown : FadeOutTimeMSPenUp;
+			return static_cast<float>((FadeOutTimeMS - DeltaTimeMS) / FadeOutTimeMS);
 		};
-		
-		for (const FPaintPacket& PaintPacket : PaintPackets)
-		{
-			const bool bPenIsDown = (PaintPacket.PenStatus & EPenStatus::CursorIsTouching) != EPenStatus::None;
-			const bool bPenIsInverted = (PaintPacket.PenStatus & EPenStatus::CursorIsInverted) != EPenStatus::None;
 
-			if (const float Opacity = GetOpacity(PaintPacket); Opacity > 0.0f)
+		OutDrawElements.PushClip(FSlateClippingZone(AllottedGeometry));
+		{
+			for (const FPaintPacket& PaintPacket : PaintPackets)
 			{
-				FLinearColor Color = bPenIsDown ? (bPenIsInverted ? ColorPenDownInverted : ColorPenDown) : ColorPenUp;
-				Color.A = Opacity;
+				if (const float Opacity = GetOpacity(PaintPacket); Opacity > 0.0f)
+				{
+					const bool bPenIsDown = (PaintPacket.PenStatus & EPenStatus::CursorIsTouching) != EPenStatus::None;
+					const bool bPenIsInverted = (PaintPacket.PenStatus & EPenStatus::CursorIsInverted) != EPenStatus::None;
+
+					FLinearColor Color = bPenIsDown ? (bPenIsInverted ? ColorPenDownInverted : ColorPenDown) : ColorPenUp;
+					Color.A = Opacity;
 
-				const float Size = FMath::Max(MinSize, PaintPacket.NormalPressure * MaxSize);
+					const float Size = FMath::Max(MinSize, PaintPacket.NormalPressure * MaxSize);
 
-				const FVector2f TransformedPosition = (PaintPacket.Position - PositionOffset - Size / 2.0f) / DPIScaleFactor;
+					const FVector2f TransformedPosition = (PaintPacket.Position - PositionOffset - Size / 2.0f) / DPIScaleFactor;
 
-				DrawCircle(TransformedPosition, Size, Color, AllottedGeometry, OutDrawElements, LayerId);
+					DrawCircle(TransformedPosition, Size, Color, AllottedGeometry, OutDrawElements, LayerId);
+				}
 			}
 		}
+		OutDrawElements.PopClip();
 
 		PaintPackets.RemoveAll([&GetOpacity](const FPaintPacket& PaintPacket) { return GetOpacity(PaintPacket) <= 0.0f; });
 
diff --git a/Engine/Plugins/Editor/StylusInput/Source/StylusInputDebugWidget/Private/StylusInputDebugPaintWidget.h b/Engine/Plugins/Editor/StylusInput/Source/StylusInputDebugWidget/Private/StylusInputDebugPaintWidget.h
--- a/Engine/Plugins/Editor/StylusInput/Source/StylusInputDebugWidget/Private/StylusInputDebugPaintWidget.h
+++ b/Engine/Plugins/Editor/StylusInput/Source/StylusInputDebugWidget/Private/StylusInputDebugPaintWidget.h
@@ -34,8 +34,8 @@ namespace UE::StylusInput::DebugWidget
 		{
 			FVector2f Position = FVector2f::ZeroVector;
 			float NormalPressure = 0.0f;
-			float TimeAddedMS = 0.0f;
 			EPenStatus PenStatus = EPenStatus::None;
+			double TimeAddedMS = 0.0;
 		};
 
 		int32 DrawPaintPackets(const FGeometry& AllottedGeometry, FSlateWindowElementList& OutDrawElements, int32 LayerId) const;
diff --git a/Engine/Plugins/Editor/StylusInput/Source/StylusInputDebugWidget/Private/StylusInputDebugWidget.cpp b/Engine/Plugins/Editor/StylusInput/Source/StylusInputDebugWidget/Private/StylusInputDebugWidget.cpp
--- a/Engine/Plugins/Editor/StylusInput/Source/StylusInputDebugWidget/Private/StylusInputDebugWidget.cpp
+++ b/Engine/Plugins/Editor/StylusInput/Source/StylusInputDebugWidget/Private/StylusInputDebugWidget.cpp
@@ -50,6 +50,9 @@ namespace UE::StylusInput::DebugWidget
 				return EActiveTimerReturnType::Stop;
 			}));
 
+		TSharedPtr<SVerticalBox> TopLeftOverlay;
+		TSharedPtr<SVerticalBox> BottomLeftOverlay;
+		
 		ChildSlot
 			.VAlign(VAlign_Fill)
 			.HAlign(HAlign_Fill)
@@ -74,153 +77,7 @@ namespace UE::StylusInput::DebugWidget
 							+ SHorizontalBox::Slot()
 							.AutoWidth()
 							[
-								SNew(SVerticalBox)
-								+ SVerticalBox::Slot()
-								.AutoHeight()
-								[
-									SNew(STextBlock)
-									.Text_Lambda([&TabletContext = LastPacketData.TabletContext]
-									{
-										return TabletContext
-											       ? FText::Format(
-												       LOCTEXT("TabletContextID", "Tablet Context ID: {0}"),
-												       FText::FromString(FString::Printf(TEXT("%x"), TabletContext->GetID())))
-											       : FText();
-									})
-								]
-								+ SVerticalBox::Slot()
-								.AutoHeight()
-								[
-									SNew(STextBlock)
-									.Text_Lambda([&TabletContext = LastPacketData.TabletContext]
-									{
-										return TabletContext
-											       ? FText::Format(
-												       LOCTEXT("TabletContextName", "Name: {0}"),
-												       FText::FromString(TabletContext->GetName()))
-											       : FText();
-									})
-								]
-								+ SVerticalBox::Slot()
-								.AutoHeight()
-								[
-									SNew(STextBlock)
-									.Text_Lambda([&TabletContext = LastPacketData.TabletContext]
-									{
-										if (TabletContext)
-										{
-											const FIntRect InputRectangle = TabletContext->GetInputRectangle();
-											return FText::Format(
-												LOCTEXT("TabletContextInputRectangle", "Input Rectangle: ({0}, {1}) x ({2}, {3})"),
-												InputRectangle.Min.X, InputRectangle.Min.Y, InputRectangle.Max.X, InputRectangle.Max.Y);
-										}
-										return FText();
-									})
-								]
-								+ SVerticalBox::Slot()
-								.AutoHeight()
-								[
-									SNew(STextBlock)
-									.Text_Lambda([&TabletContext = LastPacketData.TabletContext]
-									{
-										if (TabletContext)
-										{
-											const ETabletHardwareCapabilities Capabilities = TabletContext->GetHardwareCapabilities();
-
-											FString CapabilitiesStr;
-											if (Capabilities == ETabletHardwareCapabilities::None)
-											{
-												CapabilitiesStr = "None";
-											}
-											else
-											{
-												if ((Capabilities & ETabletHardwareCapabilities::Integrated) != ETabletHardwareCapabilities::None)
-												{
-													CapabilitiesStr += "Integrated";
-												}
-												if ((Capabilities & ETabletHardwareCapabilities::CursorMustTouch) != ETabletHardwareCapabilities::None)
-												{
-													CapabilitiesStr += CapabilitiesStr.IsEmpty() ? "CursorMustTouch" : " & CursorMustTouch";
-												}
-												if ((Capabilities & ETabletHardwareCapabilities::HardProximity) != ETabletHardwareCapabilities::None)
-												{
-													CapabilitiesStr += CapabilitiesStr.IsEmpty() ? "HardProximity" : " & HardProximity";
-												}
-												if ((Capabilities & ETabletHardwareCapabilities::CursorsHavePhysicalIds) != ETabletHardwareCapabilities::None)
-												{
-													CapabilitiesStr += CapabilitiesStr.IsEmpty() ? "CursorsHavePhysicalIds" : " & CursorsHavePhysicalIds";
-												}
-											}
-
-											if (CapabilitiesStr.IsEmpty())
-											{
-												CapabilitiesStr = FString::Format(TEXT("Unknown ({0})"), {
-													                                  static_cast<std::underlying_type_t<ETabletHardwareCapabilities>>(
-														                                  Capabilities)
-												                                  });
-											}
-
-											return FText::Format(
-												LOCTEXT("TabletContextHardwareCapabilities", "Hardware Capabilities: {0}"), FText::FromString(CapabilitiesStr));
-										}
-										return FText();
-									})
-								]
-								+ SVerticalBox::Slot()
-								.AutoHeight()
-								[
-									SNew(STextBlock)
-									.Text_Lambda([&StylusInfo = LastPacketData.StylusInfo]
-									{
-										return StylusInfo
-											       ? FText::Format(
-												       LOCTEXT("StylusID", "Stylus ID: {0}"),
-												       FText::FromString(FString::Printf(TEXT("%d"), StylusInfo->GetID())))
-											       : FText();
-									})
-								]
-								+ SVerticalBox::Slot()
-								.AutoHeight()
-								[
-									SNew(STextBlock)
-									.Text_Lambda([&StylusInfo = LastPacketData.StylusInfo]
-									{
-										return StylusInfo
-											       ? FText::Format(
-												       LOCTEXT("StylusName", "Stylus Name: {0}"),
-												       FText::FromString(StylusInfo->GetName()))
-											       : FText();
-									})
-								]
-								+ SVerticalBox::Slot()
-								.AutoHeight()
-								[
-									SNew(STextBlock)
-									.Text_Lambda([&StylusInfo = LastPacketData.StylusInfo]
-									{
-										if (StylusInfo)
-										{
-											FString ButtonsString = "<none>";
-											for (int32 Index = 0, Num = StylusInfo->GetNumButtons(); Index < Num; ++Index)
-											{
-												if (const auto* Button = StylusInfo->GetButton(Index))
-												{
-													if (Index == 0)
-													{
-														ButtonsString = Button->GetName();
-													}
-													else
-													{
-														ButtonsString += ", " + Button->GetName();
-													}
-												}
-											}
-
-											return FText::Format(LOCTEXT("StylusButtons", "Stylus Buttons: {0}"), FText::FromString(ButtonsString));
-										}
-										return FText();
-									})
-								]
+								SAssignNew(TopLeftOverlay, SVerticalBox)
 							]
 							+ SHorizontalBox::Slot()
 							.FillWidth(1.0f)
@@ -261,392 +118,12 @@ namespace UE::StylusInput::DebugWidget
 						+ SVerticalBox::Slot()
 						.AutoHeight()
 						[
-							SNew(STextBlock)
-							.Text_Lambda([&TabletContext = LastPacketData.TabletContext, &StylusInput = StylusInput, &EventHandlerThread = EventHandlerThread]
-							{
-								if (TabletContext && StylusInput)
-								{
-									return FText::Format(
-										LOCTEXT("PacketsPerSecond", "Packets Per Second: {0}"),
-										FMath::RoundToInt32(StylusInput->GetPacketsPerSecond(EventHandlerThread)));
-								}
-								return FText();
-							})
-						]
-						/*+ SVerticalBox::Slot()
-						.AutoHeight()
-						[
-							SNew(STextBlock)
-							.Text_Lambda([&bIsSet = LastPacketData.bIsSet, &CursorID = LastPacketData.Packet.CursorID]
-							{
-								return bIsSet ? FText::Format(LOCTEXT("CursorID", "Cursor ID: {0}"), CursorID) : FText();
-							})
-						]*/
-						+ SVerticalBox::Slot()
-						.AutoHeight()
-						[
-							SNew(STextBlock)
-							.Visibility_Lambda([&bIsSet = LastPacketData.bIsSet, &TabletContext = LastPacketData.TabletContext]
-							{
-								return bIsSet && TabletContext && (TabletContext->GetSupportedProperties() & ETabletSupportedProperties::TimerTick)
-								       != ETabletSupportedProperties::None
-									       ? EVisibility::Visible
-									       : EVisibility::Collapsed;
-							})
-							.Text_Lambda([&TimerTick = LastPacketData.Packet.TimerTick]
-							{
-								return FText::Format(LOCTEXT("TimerTick", "Timer Tick: {0}"), TimerTick);
-							})
-						]
-						+ SVerticalBox::Slot()
-						.AutoHeight()
-						[
-							SNew(STextBlock)
-							.Visibility_Lambda([&bIsSet = LastPacketData.bIsSet, &TabletContext = LastPacketData.TabletContext]
-							{
-								return bIsSet && TabletContext && (TabletContext->GetSupportedProperties() & ETabletSupportedProperties::SerialNumber)
-								       != ETabletSupportedProperties::None
-									       ? EVisibility::Visible
-									       : EVisibility::Collapsed;
-							})
-							.Text_Lambda([&SerialNumber = LastPacketData.Packet.SerialNumber]
-							{
-								return FText::Format(LOCTEXT("SerialNumber", "Serial Number: {0}"), SerialNumber);
-							})
-						]
-						+ SVerticalBox::Slot()
-						.AutoHeight()
-						[
-							SNew(STextBlock)
-							.Text_Lambda([&bIsSet = LastPacketData.bIsSet, &PenStatus = LastPacketData.Packet.PenStatus]
-							{
-								if (bIsSet)
-								{
-									FString PenStatusStr;
-									if (PenStatus == EPenStatus::None)
-									{
-										PenStatusStr = "None";
-									}
-									else
-									{
-										if ((PenStatus & EPenStatus::CursorIsTouching) != EPenStatus::None)
-										{
-											PenStatusStr += "CursorIsTouching";
-										}
-										if ((PenStatus & EPenStatus::CursorIsInverted) != EPenStatus::None)
-										{
-											PenStatusStr += PenStatusStr.IsEmpty() ? "CursorIsInverted" : " & CursorIsInverted";
-										}
-										if ((PenStatus & EPenStatus::NotUsed) != EPenStatus::None)
-										{
-											PenStatusStr += PenStatusStr.IsEmpty() ? "NotUsed" : " & NotUsed";
-										}
-										if ((PenStatus & EPenStatus::BarrelButtonPressed) != EPenStatus::None)
-										{
-											PenStatusStr += PenStatusStr.IsEmpty() ? "BarrelButtonPressed" : " & BarrelButtonPressed";
-										}
-									}
-
-									if (PenStatusStr.IsEmpty())
-									{
-										PenStatusStr = FString::Format(TEXT("Unknown ({0})"), {
-											                               static_cast<std::underlying_type_t<EPenStatus>>(PenStatus)
-										                               });
-									}
-
-									return FText::Format(LOCTEXT("PenStatus", "Pen Status: {0}"), FText::FromString(PenStatusStr));
-								}
-								return FText();
-							})
-						]
-						+ SVerticalBox::Slot()
-						.AutoHeight()
-						[
-							SNew(STextBlock)
-							.Visibility_Lambda([&bIsSet = LastPacketData.bIsSet, &TabletContext = LastPacketData.TabletContext]
-							{
-								return bIsSet && TabletContext && (TabletContext->GetSupportedProperties() & ETabletSupportedProperties::X)
-								       != ETabletSupportedProperties::None
-									       ? EVisibility::Visible
-									       : EVisibility::Collapsed;
-							})
-							.Text_Lambda([&X = LastPacketData.Packet.X]
-							{
-								return FText::Format(LOCTEXT("X", "X: {0}"), X);
-							})
-						]
-						+ SVerticalBox::Slot()
-						.AutoHeight()
-						[
-							SNew(STextBlock)
-							.Visibility_Lambda([&bIsSet = LastPacketData.bIsSet, &TabletContext = LastPacketData.TabletContext]
-							{
-								return bIsSet && TabletContext && (TabletContext->GetSupportedProperties() & ETabletSupportedProperties::Y)
-								       != ETabletSupportedProperties::None
-									       ? EVisibility::Visible
-									       : EVisibility::Collapsed;
-							})
-							.Text_Lambda([&Y = LastPacketData.Packet.Y]
-							{
-								return FText::Format(LOCTEXT("Y", "Y: {0}"), Y);
-							})
-						]
-						+ SVerticalBox::Slot()
-						.AutoHeight()
-						[
-							SNew(STextBlock)
-							.Visibility_Lambda([&bIsSet = LastPacketData.bIsSet, &TabletContext = LastPacketData.TabletContext]
-							{
-								return bIsSet && TabletContext && (TabletContext->GetSupportedProperties() & ETabletSupportedProperties::Z)
-								       != ETabletSupportedProperties::None
-									       ? EVisibility::Visible
-									       : EVisibility::Collapsed;
-							})
-							.Text_Lambda([&Z = LastPacketData.Packet.Z]
-							{
-								return FText::Format(LOCTEXT("Z", "Z: {0}"), Z);
-							})
-						]
-						+ SVerticalBox::Slot()
-						.AutoHeight()
-						[
-							SNew(STextBlock)
-							.Visibility_Lambda([&bIsSet = LastPacketData.bIsSet, &TabletContext = LastPacketData.TabletContext]
-							{
-								return bIsSet && TabletContext && (TabletContext->GetSupportedProperties() & ETabletSupportedProperties::NormalPressure)
-								       != ETabletSupportedProperties::None
-									       ? EVisibility::Visible
-									       : EVisibility::Collapsed;
-							})
-							.Text_Lambda([&NormalPressure = LastPacketData.Packet.NormalPressure]
-							{
-								return FText::Format(LOCTEXT("NormalPressure", "Normal Pressure: {0}"), NormalPressure);
-							})
-						]
-						+ SVerticalBox::Slot()
-						.AutoHeight()
-						[
-							SNew(STextBlock)
-							.Visibility_Lambda([&bIsSet = LastPacketData.bIsSet, &TabletContext = LastPacketData.TabletContext]
-							{
-								return bIsSet && TabletContext && (TabletContext->GetSupportedProperties() & ETabletSupportedProperties::TangentPressure)
-								       != ETabletSupportedProperties::None
-									       ? EVisibility::Visible
-									       : EVisibility::Collapsed;
-							})
-							.Text_Lambda([&TangentPressure = LastPacketData.Packet.TangentPressure]
-							{
-								return FText::Format(LOCTEXT("TangentPressure", "Tangent Pressure: {0}"), TangentPressure);
-							})
-						]
-						+ SVerticalBox::Slot()
-						.AutoHeight()
-						[
-							SNew(STextBlock)
-							.Visibility_Lambda([&bIsSet = LastPacketData.bIsSet, &TabletContext = LastPacketData.TabletContext]
-							{
-								return bIsSet && TabletContext && (TabletContext->GetSupportedProperties() & ETabletSupportedProperties::ButtonPressure)
-								       != ETabletSupportedProperties::None
-									       ? EVisibility::Visible
-									       : EVisibility::Collapsed;
-							})
-							.Text_Lambda([&ButtonPressure = LastPacketData.Packet.ButtonPressure]
-							{
-								return FText::Format(LOCTEXT("ButtonPressure", "Button Pressure: {0}"), ButtonPressure);
-							})
-						]
-						+ SVerticalBox::Slot()
-						.AutoHeight()
-						[
-							SNew(STextBlock)
-							.Visibility_Lambda([&bIsSet = LastPacketData.bIsSet, &TabletContext = LastPacketData.TabletContext]
-							{
-								return bIsSet && TabletContext && (TabletContext->GetSupportedProperties() & ETabletSupportedProperties::XTiltOrientation)
-								       != ETabletSupportedProperties::None
-									       ? EVisibility::Visible
-									       : EVisibility::Collapsed;
-							})
-							.Text_Lambda([&XTiltOrientation = LastPacketData.Packet.XTiltOrientation]
-							{
-								return FText::Format(LOCTEXT("XTiltOrientation", "X Tilt Orientation: {0}"), XTiltOrientation);
-							})
-						]
-						+ SVerticalBox::Slot()
-						.AutoHeight()
-						[
-							SNew(STextBlock)
-							.Visibility_Lambda([&bIsSet = LastPacketData.bIsSet, &TabletContext = LastPacketData.TabletContext]
-							{
-								return bIsSet && TabletContext && (TabletContext->GetSupportedProperties() & ETabletSupportedProperties::YTiltOrientation)
-								       != ETabletSupportedProperties::None
-									       ? EVisibility::Visible
-									       : EVisibility::Collapsed;
-							})
-							.Text_Lambda([&YTiltOrientation = LastPacketData.Packet.YTiltOrientation]
-							{
-								return FText::Format(LOCTEXT("YTiltOrientation", "Y Tilt Orientation: {0}"), YTiltOrientation);
-							})
-						]
-						+ SVerticalBox::Slot()
-						.AutoHeight()
-						[
-							SNew(STextBlock)
-							.Visibility_Lambda([&bIsSet = LastPacketData.bIsSet, &TabletContext = LastPacketData.TabletContext]
-							{
-								return bIsSet && TabletContext && (TabletContext->GetSupportedProperties() & ETabletSupportedProperties::AzimuthOrientation)
-								       != ETabletSupportedProperties::None
-									       ? EVisibility::Visible
-									       : EVisibility::Collapsed;
-							})
-							.Text_Lambda([&AzimuthOrientation = LastPacketData.Packet.AzimuthOrientation]
-							{
-								return FText::Format(LOCTEXT("AzimuthOrientation", "Azimuth Orientation: {0}"), AzimuthOrientation);
-							})
-						]
-						+ SVerticalBox::Slot()
-						.AutoHeight()
-						[
-							SNew(STextBlock)
-							.Visibility_Lambda([&bIsSet = LastPacketData.bIsSet, &TabletContext = LastPacketData.TabletContext]
-							{
-								return bIsSet && TabletContext && (TabletContext->GetSupportedProperties() & ETabletSupportedProperties::AltitudeOrientation)
-								       != ETabletSupportedProperties::None
-									       ? EVisibility::Visible
-									       : EVisibility::Collapsed;
-							})
-							.Text_Lambda([&AltitudeOrientation = LastPacketData.Packet.AltitudeOrientation]
-							{
-								return FText::Format(LOCTEXT("AltitudeOrientation", "Altitude Orientation: {0}"), AltitudeOrientation);
-							})
-						]
-						+ SVerticalBox::Slot()
-						.AutoHeight()
-						[
-							SNew(STextBlock)
-							.Visibility_Lambda([&bIsSet = LastPacketData.bIsSet, &TabletContext = LastPacketData.TabletContext]
-							{
-								return bIsSet && TabletContext && (TabletContext->GetSupportedProperties() & ETabletSupportedProperties::TwistOrientation)
-								       != ETabletSupportedProperties::None
-									       ? EVisibility::Visible
-									       : EVisibility::Collapsed;
-							})
-							.Text_Lambda([&TwistOrientation = LastPacketData.Packet.TwistOrientation]
-							{
-								return FText::Format(LOCTEXT("TwistOrientation", "Twist Orientation: {0}"), TwistOrientation);
-							})
-						]
-						+ SVerticalBox::Slot()
-						.AutoHeight()
-						[
-							SNew(STextBlock)
-							.Visibility_Lambda([&bIsSet = LastPacketData.bIsSet, &TabletContext = LastPacketData.TabletContext]
-							{
-								return bIsSet && TabletContext && (TabletContext->GetSupportedProperties() & ETabletSupportedProperties::PitchRotation)
-								       != ETabletSupportedProperties::None
-									       ? EVisibility::Visible
-									       : EVisibility::Collapsed;
-							})
-							.Text_Lambda([&PitchRotation = LastPacketData.Packet.PitchRotation]
-							{
-								return FText::Format(LOCTEXT("PitchRotation", "Pitch Rotation: {0}"), PitchRotation);
-							})
-						]
-						+ SVerticalBox::Slot()
-						.AutoHeight()
-						[
-							SNew(STextBlock)
-							.Visibility_Lambda([&bIsSet = LastPacketData.bIsSet, &TabletContext = LastPacketData.TabletContext]
-							{
-								return bIsSet && TabletContext && (TabletContext->GetSupportedProperties() & ETabletSupportedProperties::RollRotation)
-								       != ETabletSupportedProperties::None
-									       ? EVisibility::Visible
-									       : EVisibility::Collapsed;
-							})
-							.Text_Lambda([&RollRotation = LastPacketData.Packet.RollRotation]
-							{
-								return FText::Format(LOCTEXT("RollRotation", "Roll Rotation: {0}"), RollRotation);
-							})
-						]
-						+ SVerticalBox::Slot()
-						.AutoHeight()
-						[
-							SNew(STextBlock)
-							.Visibility_Lambda([&bIsSet = LastPacketData.bIsSet, &TabletContext = LastPacketData.TabletContext]
-							{
-								return bIsSet && TabletContext && (TabletContext->GetSupportedProperties() & ETabletSupportedProperties::YawRotation)
-								       != ETabletSupportedProperties::None
-									       ? EVisibility::Visible
-									       : EVisibility::Collapsed;
-							})
-							.Text_Lambda([&YawRotation = LastPacketData.Packet.YawRotation]
-							{
-								return FText::Format(LOCTEXT("YawRotation", "Yaw Rotation: {0}"), YawRotation);
-							})
-						]
-						+ SVerticalBox::Slot()
-						.AutoHeight()
-						[
-							SNew(STextBlock)
-							.Visibility_Lambda([&bIsSet = LastPacketData.bIsSet, &TabletContext = LastPacketData.TabletContext]
-							{
-								return bIsSet && TabletContext && (TabletContext->GetSupportedProperties() & ETabletSupportedProperties::Width)
-								       != ETabletSupportedProperties::None
-									       ? EVisibility::Visible
-									       : EVisibility::Collapsed;
-							})
-							.Text_Lambda([&Width = LastPacketData.Packet.Width]
-							{
-								return FText::Format(LOCTEXT("Width", "Width: {0}"), Width);
-							})
-						]
-						+ SVerticalBox::Slot()
-						.AutoHeight()
-						[
-							SNew(STextBlock)
-							.Visibility_Lambda([&bIsSet = LastPacketData.bIsSet, &TabletContext = LastPacketData.TabletContext]
-							{
-								return bIsSet && TabletContext && (TabletContext->GetSupportedProperties() & ETabletSupportedProperties::Height)
-								       != ETabletSupportedProperties::None
-									       ? EVisibility::Visible
-									       : EVisibility::Collapsed;
-							})
-							.Text_Lambda([&Height = LastPacketData.Packet.Height]
-							{
-								return FText::Format(LOCTEXT("Height", "Height: {0}"), Height);
-							})
-						]
-						+ SVerticalBox::Slot()
-						.AutoHeight()
-						[
-							SNew(STextBlock)
-							.Visibility_Lambda([&bIsSet = LastPacketData.bIsSet, &TabletContext = LastPacketData.TabletContext]
-							{
-								return bIsSet && TabletContext && (TabletContext->GetSupportedProperties() &
-									       ETabletSupportedProperties::FingerContactConfidence)
-								       != ETabletSupportedProperties::None
-									       ? EVisibility::Visible
-									       : EVisibility::Collapsed;
-							})
-							.Text_Lambda([&FingerContactConfidence = LastPacketData.Packet.FingerContactConfidence]
-							{
-								return FText::Format(LOCTEXT("FingerContactConfidence", "Finger Contact Confidence: {0}"), FingerContactConfidence);
-							})
-						]
-						+ SVerticalBox::Slot()
-						.AutoHeight()
-						[
-							SNew(STextBlock)
-							.Visibility_Lambda([&bIsSet = LastPacketData.bIsSet, &TabletContext = LastPacketData.TabletContext]
-							{
-								return bIsSet && TabletContext && (TabletContext->GetSupportedProperties() & ETabletSupportedProperties::DeviceContactID)
-								       != ETabletSupportedProperties::None
-									       ? EVisibility::Visible
-									       : EVisibility::Collapsed;
-							})
-							.Text_Lambda([&DeviceContactID = LastPacketData.Packet.DeviceContactID]
-							{
-								return FText::Format(LOCTEXT("DeviceContactID", "Device Contact ID: {0}"), DeviceContactID);
-							})
+							SNew(SHorizontalBox)
+							+ SHorizontalBox::Slot()
+							.AutoWidth()
+							[
+								SAssignNew(BottomLeftOverlay, SVerticalBox)
+							]
 						]
 					]
 				]
@@ -659,6 +136,331 @@ namespace UE::StylusInput::DebugWidget
 					.VScrollBar(SNew(SScrollBar).Orientation(Orient_Vertical).AlwaysShowScrollbar(true).Thickness(10))
 				]
 			];
+
+		auto AddToOverlay = [](const TSharedPtr<SVerticalBox>& Box, TFunction<FText()>&& Text)
+		{
+			Box->AddSlot()
+				.AutoHeight()
+				[
+					SNew(STextBlock)
+					.Text_Lambda(MoveTemp(Text))
+				];
+		};
+
+		auto AddToOverlayWithVisibility = [&bIsSet = LastPacketData.bIsSet, &TabletContext = LastPacketData.TabletContext](const TSharedPtr<SVerticalBox>& Box, ETabletSupportedProperties Property, TFunction<FText()>&& Text)
+		{
+			Box->AddSlot()
+				.AutoHeight()
+				[
+					SNew(STextBlock)
+						.Visibility_Lambda([&bIsSet, &TabletContext, Property]
+							{
+								return bIsSet && TabletContext && (TabletContext->GetSupportedProperties() & Property)
+								       != ETabletSupportedProperties::None
+									       ? EVisibility::Visible
+									       : EVisibility::Collapsed;
+							})
+						.Text_Lambda(MoveTemp(Text))
+				];
+		};
+
+		AddToOverlay(TopLeftOverlay,
+		             [&TabletContext = LastPacketData.TabletContext]
+		             {
+			             return TabletContext
+				                    ? FText::Format(
+					                    LOCTEXT("TabletContextID", "Tablet Context ID: {0}"),
+					                    FText::FromString(FString::Printf(TEXT("%x"), TabletContext->GetID())))
+				                    : FText();
+		             });
+
+		AddToOverlay(TopLeftOverlay,
+		             [&TabletContext = LastPacketData.TabletContext]
+		             {
+			             return TabletContext
+				                    ? FText::Format(
+					                    LOCTEXT("TabletContextName", "Name: {0}"),
+					                    FText::FromString(TabletContext->GetName()))
+				                    : FText();
+		             });
+
+		AddToOverlay(TopLeftOverlay,
+		             [&TabletContext = LastPacketData.TabletContext]
+		             {
+			             if (TabletContext)
+			             {
+				             const FIntRect InputRectangle = TabletContext->GetInputRectangle();
+				             return FText::Format(
+					             LOCTEXT("TabletContextInputRectangle", "Input Rectangle: ({0}, {1}) x ({2}, {3})"),
+					             InputRectangle.Min.X, InputRectangle.Min.Y, InputRectangle.Max.X, InputRectangle.Max.Y);
+			             }
+			             return FText();
+		             });
+
+		AddToOverlay(TopLeftOverlay,
+		             [&TabletContext = LastPacketData.TabletContext]
+		             {
+			             if (TabletContext)
+			             {
+				             const ETabletHardwareCapabilities Capabilities = TabletContext->GetHardwareCapabilities();
+
+				             FString CapabilitiesStr;
+				             if (Capabilities == ETabletHardwareCapabilities::None)
+				             {
+					             CapabilitiesStr = "None";
+				             }
+				             else
+				             {
+					             if ((Capabilities & ETabletHardwareCapabilities::Integrated) != ETabletHardwareCapabilities::None)
+					             {
+						             CapabilitiesStr += "Integrated";
+					             }
+					             if ((Capabilities & ETabletHardwareCapabilities::CursorMustTouch) != ETabletHardwareCapabilities::None)
+					             {
+						             CapabilitiesStr += CapabilitiesStr.IsEmpty() ? "CursorMustTouch" : " & CursorMustTouch";
+					             }
+					             if ((Capabilities & ETabletHardwareCapabilities::HardProximity) != ETabletHardwareCapabilities::None)
+					             {
+						             CapabilitiesStr += CapabilitiesStr.IsEmpty() ? "HardProximity" : " & HardProximity";
+					             }
+					             if ((Capabilities & ETabletHardwareCapabilities::CursorsHavePhysicalIds) != ETabletHardwareCapabilities::None)
+					             {
+						             CapabilitiesStr += CapabilitiesStr.IsEmpty() ? "CursorsHavePhysicalIds" : " & CursorsHavePhysicalIds";
+					             }
+				             }
+
+				             if (CapabilitiesStr.IsEmpty())
+				             {
+					             CapabilitiesStr = FString::Format(TEXT("Unknown ({0})"), {
+						                                               static_cast<std::underlying_type_t<ETabletHardwareCapabilities>>(
+							                                               Capabilities)
+					                                               });
+				             }
+
+				             return FText::Format(
+					             LOCTEXT("TabletContextHardwareCapabilities", "Hardware Capabilities: {0}"), FText::FromString(CapabilitiesStr));
+			             }
+			             return FText();
+		             });
+
+		AddToOverlay(TopLeftOverlay,
+		             [&StylusInfo = LastPacketData.StylusInfo]
+		             {
+			             return StylusInfo
+				                    ? FText::Format(
+					                    LOCTEXT("StylusID", "Stylus ID: {0}"),
+					                    FText::FromString(FString::Printf(TEXT("%d"), StylusInfo->GetID())))
+				                    : FText();
+		             });
+
+		AddToOverlay(TopLeftOverlay,
+		             [&StylusInfo = LastPacketData.StylusInfo]
+		             {
+			             return StylusInfo
+				                    ? FText::Format(
+					                    LOCTEXT("StylusName", "Stylus Name: {0}"),
+					                    FText::FromString(StylusInfo->GetName()))
+				                    : FText();
+		             });
+		
+		AddToOverlay(TopLeftOverlay,
+		             [&StylusInfo = LastPacketData.StylusInfo]
+		             {
+			             if (StylusInfo)
+			             {
+				             FString ButtonsString = "<none>";
+				             for (int32 Index = 0, Num = StylusInfo->GetNumButtons(); Index < Num; ++Index)
+				             {
+					             if (const auto* Button = StylusInfo->GetButton(Index))
+					             {
+						             if (Index == 0)
+						             {
+							             ButtonsString = Button->GetName();
+						             }
+						             else
+						             {
+							             ButtonsString += ", " + Button->GetName();
+						             }
+					             }
+				             }
+				             return FText::Format(LOCTEXT("StylusButtons", "Stylus Buttons: {0}"), FText::FromString(ButtonsString));
+			             }
+			             return FText();
+		             });
+
+		AddToOverlay(BottomLeftOverlay,
+		             [&TabletContext = LastPacketData.TabletContext, &StylusInput = StylusInput, &EventHandlerThread = EventHandlerThread]
+		             {
+			             if (TabletContext && StylusInput)
+			             {
+				             return FText::Format(
+					             LOCTEXT("PacketsPerSecond", "Packets Per Second: {0}"),
+					             FMath::RoundToInt32(StylusInput->GetPacketsPerSecond(EventHandlerThread)));
+			             }
+			             return FText();
+		             });
+
+		AddToOverlayWithVisibility(BottomLeftOverlay, ETabletSupportedProperties::TimerTick,
+		                           [&TimerTick = LastPacketData.Packet.TimerTick]
+		                           {
+			                           return FText::Format(LOCTEXT("TimerTick", "Timer Tick: {0}"), TimerTick);
+		                           });
+
+		AddToOverlayWithVisibility(BottomLeftOverlay, ETabletSupportedProperties::SerialNumber,
+		                           [&SerialNumber = LastPacketData.Packet.SerialNumber]
+		                           {
+			                           return FText::Format(LOCTEXT("SerialNumber", "Serial Number: {0}"), SerialNumber);
+		                           });
+
+		AddToOverlay(BottomLeftOverlay,
+		             [&bIsSet = LastPacketData.bIsSet, &PenStatus = LastPacketData.Packet.PenStatus]
+		             {
+			             if (bIsSet)
+			             {
+				             FString PenStatusStr;
+				             if (PenStatus == EPenStatus::None)
+				             {
+					             PenStatusStr = "None";
+				             }
+				             else
+				             {
+					             if ((PenStatus & EPenStatus::CursorIsTouching) != EPenStatus::None)
+					             {
+						             PenStatusStr += "CursorIsTouching";
+					             }
+					             if ((PenStatus & EPenStatus::CursorIsInverted) != EPenStatus::None)
+					             {
+						             PenStatusStr += PenStatusStr.IsEmpty() ? "CursorIsInverted" : " & CursorIsInverted";
+					             }
+					             if ((PenStatus & EPenStatus::NotUsed) != EPenStatus::None)
+					             {
+						             PenStatusStr += PenStatusStr.IsEmpty() ? "NotUsed" : " & NotUsed";
+					             }
+					             if ((PenStatus & EPenStatus::BarrelButtonPressed) != EPenStatus::None)
+					             {
+						             PenStatusStr += PenStatusStr.IsEmpty() ? "BarrelButtonPressed" : " & BarrelButtonPressed";
+					             }
+				             }
+
+				             if (PenStatusStr.IsEmpty())
+				             {
+					             PenStatusStr = FString::Format(TEXT("Unknown ({0})"), {
+						                                            static_cast<std::underlying_type_t<EPenStatus>>(PenStatus)
+					                                            });
+				             }
+
+				             return FText::Format(LOCTEXT("PenStatus", "Pen Status: {0}"), FText::FromString(PenStatusStr));
+			             }
+			             return FText();
+		             });
+
+		AddToOverlayWithVisibility(BottomLeftOverlay, ETabletSupportedProperties::X,
+		                           [&X = LastPacketData.Packet.X]
+		                           {
+			                           return FText::Format(LOCTEXT("X", "X: {0}"), X);
+		                           });
+
+		AddToOverlayWithVisibility(BottomLeftOverlay, ETabletSupportedProperties::Y,
+		                           [&Y = LastPacketData.Packet.Y]
+		                           {
+			                           return FText::Format(LOCTEXT("Y", "Y: {0}"), Y);
+		                           });
+
+		AddToOverlayWithVisibility(BottomLeftOverlay, ETabletSupportedProperties::Z,
+		                           [&Z = LastPacketData.Packet.Z]
+		                           {
+			                           return FText::Format(LOCTEXT("Z", "Z: {0}"), Z);
+		                           });
+
+		AddToOverlayWithVisibility(BottomLeftOverlay, ETabletSupportedProperties::NormalPressure,
+		                           [&NormalPressure = LastPacketData.Packet.NormalPressure]
+		                           {
+			                           return FText::Format(LOCTEXT("NormalPressure", "Normal Pressure: {0}"), NormalPressure);
+		                           });
+
+		AddToOverlayWithVisibility(BottomLeftOverlay, ETabletSupportedProperties::TangentPressure,
+		                           [&TangentPressure = LastPacketData.Packet.TangentPressure]
+		                           {
+			                           return FText::Format(LOCTEXT("TangentPressure", "Tangent Pressure: {0}"), TangentPressure);
+		                           });
+
+		AddToOverlayWithVisibility(BottomLeftOverlay, ETabletSupportedProperties::ButtonPressure,
+		                           [&ButtonPressure = LastPacketData.Packet.ButtonPressure]
+		                           {
+			                           return FText::Format(LOCTEXT("ButtonPressure", "Button Pressure: {0}"), ButtonPressure);
+		                           });
+
+		AddToOverlayWithVisibility(BottomLeftOverlay, ETabletSupportedProperties::XTiltOrientation,
+		                           [&XTiltOrientation = LastPacketData.Packet.XTiltOrientation]
+		                           {
+			                           return FText::Format(LOCTEXT("XTiltOrientation", "X Tilt Orientation: {0}"), XTiltOrientation);
+		                           });
+
+		AddToOverlayWithVisibility(BottomLeftOverlay, ETabletSupportedProperties::YTiltOrientation,
+		                           [&YTiltOrientation = LastPacketData.Packet.YTiltOrientation]
+		                           {
+			                           return FText::Format(LOCTEXT("YTiltOrientation", "Y Tilt Orientation: {0}"), YTiltOrientation);
+		                           });
+
+		AddToOverlayWithVisibility(BottomLeftOverlay, ETabletSupportedProperties::AzimuthOrientation,
+		                           [&AzimuthOrientation = LastPacketData.Packet.AzimuthOrientation]
+		                           {
+			                           return FText::Format(LOCTEXT("AzimuthOrientation", "Azimuth Orientation: {0}"), AzimuthOrientation);
+		                           });
+
+		AddToOverlayWithVisibility(BottomLeftOverlay, ETabletSupportedProperties::AltitudeOrientation,
+		                           [&AltitudeOrientation = LastPacketData.Packet.AltitudeOrientation]
+		                           {
+			                           return FText::Format(LOCTEXT("AltitudeOrientation", "Altitude Orientation: {0}"), AltitudeOrientation);
+		                           });
+
+		AddToOverlayWithVisibility(BottomLeftOverlay, ETabletSupportedProperties::TwistOrientation,
+		                           [&TwistOrientation = LastPacketData.Packet.TwistOrientation]
+		                           {
+			                           return FText::Format(LOCTEXT("TwistOrientation", "Twist Orientation: {0}"), TwistOrientation);
+		                           });
+
+		AddToOverlayWithVisibility(BottomLeftOverlay, ETabletSupportedProperties::PitchRotation,
+		                           [&PitchRotation = LastPacketData.Packet.PitchRotation]
+		                           {
+			                           return FText::Format(LOCTEXT("PitchRotation", "Pitch Rotation: {0}"), PitchRotation);
+		                           });
+
+		AddToOverlayWithVisibility(BottomLeftOverlay, ETabletSupportedProperties::RollRotation,
+		                           [&RollRotation = LastPacketData.Packet.RollRotation]
+		                           {
+			                           return FText::Format(LOCTEXT("RollRotation", "Roll Rotation: {0}"), RollRotation);
+		                           });
+
+		AddToOverlayWithVisibility(BottomLeftOverlay, ETabletSupportedProperties::YawRotation,
+		                           [&YawRotation = LastPacketData.Packet.YawRotation]
+		                           {
+			                           return FText::Format(LOCTEXT("YawRotation", "Yaw Rotation: {0}"), YawRotation);
+		                           });
+
+		AddToOverlayWithVisibility(BottomLeftOverlay, ETabletSupportedProperties::Width,
+								   [&Width = LastPacketData.Packet.Width]
+								   {
+									   return FText::Format(LOCTEXT("Width", "Width: {0}"), Width);
+								   });
+
+		AddToOverlayWithVisibility(BottomLeftOverlay, ETabletSupportedProperties::Height,
+								   [&Height = LastPacketData.Packet.Height]
+								   {
+									   return FText::Format(LOCTEXT("Height", "Height: {0}"), Height);
+								   });
+
+		AddToOverlayWithVisibility(BottomLeftOverlay, ETabletSupportedProperties::FingerContactConfidence,
+								   [&FingerContactConfidence = LastPacketData.Packet.FingerContactConfidence]
+								   {
+									   return FText::Format(LOCTEXT("FingerContactConfidence", "Finger Contact Confidence: {0}"), FingerContactConfidence);
+								   });
+
+		AddToOverlayWithVisibility(BottomLeftOverlay, ETabletSupportedProperties::DeviceContactID,
+		                           [&DeviceContactID = LastPacketData.Packet.DeviceContactID]
+		                           {
+			                           return FText::Format(LOCTEXT("DeviceContactID", "Device Contact ID: {0}"), DeviceContactID);
+		                           });
 	}
 
 	void SStylusInputDebugWidget::AcquireStylusInput()
