StylusInput: Update icon

diff --git a/Engine/Plugins/Editor/StylusInput/Source/StylusInputDebugWidget/Private/StylusInputDebugWidgetModule.cpp b/Engine/Plugins/Editor/StylusInput/Source/StylusInputDebugWidget/Private/StylusInputDebugWidgetModule.cpp
--- a/Engine/Plugins/Editor/StylusInput/Source/StylusInputDebugWidget/Private/StylusInputDebugWidgetModule.cpp
+++ b/Engine/Plugins/Editor/StylusInput/Source/StylusInputDebugWidget/Private/StylusInputDebugWidgetModule.cpp
@@ -3,6 +3,8 @@
 #include <WorkspaceMenuStructure.h>
 #include <WorkspaceMenuStructureModule.h>
 #include <Modules/ModuleManager.h>
+#include <Styling/SlateStyle.h>
+#include <Styling/SlateStyleRegistry.h>
 #include <Widgets/Docking/SDockTab.h>
 
 #include "StylusInputDebugWidget.h"
@@ -23,19 +25,24 @@ namespace UE::StylusInput::DebugWidget
 	private:
 		TSharedRef<SDockTab> MakeDebugWidgetTab(const FSpawnTabArgs&);
 		TSharedRef<SWidget> GetDebugWidget();
+		void SetupStyleSet();
+		void ResetStyleSet();
 
 		bool bHasRegisteredTabSpawners = false;
 		TWeakPtr<SStylusInputDebugWidget> DebugWidgetPtr;
+		TSharedPtr<FSlateStyleSet> StyleSet;
 	};
 
 	void FStylusInputDebugWidgetModule::StartupModule()
 	{
+		SetupStyleSet();
 		RegisterTabSpawners(nullptr);
 	}
 
 	void FStylusInputDebugWidgetModule::ShutdownModule()
 	{
 		UnregisterTabSpawners();
+		ResetStyleSet();
 	}
 
 	void FStylusInputDebugWidgetModule::RegisterTabSpawners(const TSharedPtr<FWorkspaceItem>& WorkspaceGroup)
@@ -50,7 +57,7 @@ namespace UE::StylusInput::DebugWidget
 				.SetDisplayName(LOCTEXT("DebugWidgetTitle", "Stylus Input Debug"))
 				.SetTooltipText(LOCTEXT("DebugWidgetTooltip", "Open a debug widget to verify stylus input event handling."))
 				.SetGroup(WorkspaceMenu::GetMenuStructure().GetDeveloperToolsDebugCategory())
-				.SetIcon(FSlateIcon(FAppStyle::GetAppStyleSetName(), "StylusInputDebug.TabIcon"));
+				.SetIcon(FSlateIcon(StyleSet->GetStyleSetName(), "StylusInput.Small"));
 
 		if (WorkspaceGroup.IsValid())
 		{
@@ -86,6 +93,25 @@ namespace UE::StylusInput::DebugWidget
 
 		return DebugWidget.ToSharedRef();
 	}
+
+	void FStylusInputDebugWidgetModule::SetupStyleSet()
+	{
+		static FName StylusInputDebugWidgetStyle(TEXT("StylusInputDebugWidgetStyle"));
+		StyleSet = MakeShared<FSlateStyleSet>(StylusInputDebugWidgetStyle);
+
+		StyleSet->SetContentRoot(FPaths::EnginePluginsDir() / TEXT("Editor/StylusInput/Resources"));
+
+		StyleSet->Set("StylusInput.Small",
+		              new FSlateVectorImageBrush(StyleSet->RootToContentDir(TEXT("StylusInput_16"), TEXT(".svg")), {16.0f, 16.0f}));
+
+		FSlateStyleRegistry::RegisterSlateStyle(*StyleSet);
+	}
+
+	void FStylusInputDebugWidgetModule::ResetStyleSet()
+	{
+		FSlateStyleRegistry::UnRegisterSlateStyle(*StyleSet);
+		StyleSet.Reset();
+	}
 }
 
 #undef LOCTEXT_NAMESPACE
diff --git a/Engine/Source/Editor/EditorStyle/Private/StarshipStyle.cpp b/Engine/Source/Editor/EditorStyle/Private/StarshipStyle.cpp
--- a/Engine/Source/Editor/EditorStyle/Private/StarshipStyle.cpp
+++ b/Engine/Source/Editor/EditorStyle/Private/StarshipStyle.cpp
@@ -945,7 +945,6 @@ void FStarshipEditorStyle::FStyle::SetupGeneralStyles()
 		Set("CollisionAnalyzer.TabIcon", new IMAGE_BRUSH_SVG("Starship/Common/Collision", Icon16x16));
 		Set("ObjectBrowser.TabIcon", new IMAGE_BRUSH_SVG( "Starship/Common/ObjectsBrowser", Icon16x16 ) );
 		Set("PixelInspector.TabIcon", new IMAGE_BRUSH_SVG( "Starship/Common/PixelInspector", Icon16x16 ) );
-		Set("StylusInputDebug.TabIcon", new IMAGE_BRUSH_SVG( "Starship/Common/StylusInputBug", Icon16x16 ) );
 	}
 
 	{
