StylusInput: Provide stylus input instance in event handler calls

diff --git a/Engine/Plugins/Editor/StylusInput/Source/StylusInput/Private/DeprecatedWindowsStylusInputInterface.cpp b/Engine/Plugins/Editor/StylusInput/Source/StylusInput/Private/DeprecatedWindowsStylusInputInterface.cpp
--- a/Engine/Plugins/Editor/StylusInput/Source/StylusInput/Private/DeprecatedWindowsStylusInputInterface.cpp
+++ b/Engine/Plugins/Editor/StylusInput/Source/StylusInput/Private/DeprecatedWindowsStylusInputInterface.cpp
@@ -151,27 +151,26 @@ void FDeprecatedWindowsStylusInputInterface::FDeprecatedStylusInputDevice::Tick(
 	PRAGMA_ENABLE_DEPRECATION_WARNINGS
 }
 
-FDeprecatedWindowsStylusInputInterface::FStylusInputEventHandler::FStylusInputEventHandler(UE::StylusInput::IStylusInputInstance* Instance, TArray<FDeprecatedStylusInputDevice>& TabletContexts)
-	: Instance(Instance)
-	, TabletContexts(TabletContexts)
+FDeprecatedWindowsStylusInputInterface::FStylusInputEventHandler::FStylusInputEventHandler(TArray<FDeprecatedStylusInputDevice>& TabletContexts)
+	: TabletContexts(TabletContexts)
 {
 }
 
-void FDeprecatedWindowsStylusInputInterface::FStylusInputEventHandler::OnPacket(const UE::StylusInput::FStylusInputPacket& Packet)
+void FDeprecatedWindowsStylusInputInterface::FStylusInputEventHandler::OnPacket(const UE::StylusInput::FStylusInputPacket& Packet,
+                                                                                UE::StylusInput::IStylusInputInstance* Instance)
 {
-	FDeprecatedStylusInputDevice* TabletContext = GetTabletContext(Packet.TabletContextID);
+	FDeprecatedStylusInputDevice* TabletContext = GetTabletContext(Packet.TabletContextID, Instance);
 
 	TabletContext->LastPacket = Packet;
 	TabletContext->SetDirty();
 }
 
-void FDeprecatedWindowsStylusInputInterface::FStylusInputEventHandler::OnDebugEvent(const FString& Message)
+void FDeprecatedWindowsStylusInputInterface::FStylusInputEventHandler::OnDebugEvent(const FString& Message, UE::StylusInput::IStylusInputInstance* Instance)
 {
-	// UE_LOG(LogStylusInput, Warning, TEXT("%s"), *Message);
 }
 
 FDeprecatedWindowsStylusInputInterface::FDeprecatedStylusInputDevice* FDeprecatedWindowsStylusInputInterface::FStylusInputEventHandler::GetTabletContext(
-	uint32 TabletContextId)
+	uint32 TabletContextId, UE::StylusInput::IStylusInputInstance* Instance)
 {
 	FDeprecatedStylusInputDevice* TabletContext = Algo::FindByPredicate(
 		TabletContexts, [TabletContextId](const FDeprecatedStylusInputDevice& Value)
@@ -193,7 +192,7 @@ FDeprecatedWindowsStylusInputInterface::FDeprecatedStylusInputDevice* FDeprecate
 FDeprecatedWindowsStylusInputInterface::FStylusInputInstanceWrapper::FStylusInputInstanceWrapper(SWindow* Window,
                                                                                                  TArray<FDeprecatedStylusInputDevice>& TabletContexts)
 	: Instance(UE::StylusInput::CreateInstance(*Window))
-	, EventHandler(Instance, TabletContexts)
+	, EventHandler(TabletContexts)
 {
 	if (Instance)
 	{
diff --git a/Engine/Plugins/Editor/StylusInput/Source/StylusInput/Private/DeprecatedWindowsStylusInputInterface.h b/Engine/Plugins/Editor/StylusInput/Source/StylusInput/Private/DeprecatedWindowsStylusInputInterface.h
--- a/Engine/Plugins/Editor/StylusInput/Source/StylusInput/Private/DeprecatedWindowsStylusInputInterface.h
+++ b/Engine/Plugins/Editor/StylusInput/Source/StylusInput/Private/DeprecatedWindowsStylusInputInterface.h
@@ -52,16 +52,15 @@ private:
 	class FStylusInputEventHandler : public UE::StylusInput::IStylusInputEventHandler
 	{
 	public:
-		explicit FStylusInputEventHandler(UE::StylusInput::IStylusInputInstance* Instance, TArray<FDeprecatedStylusInputDevice>& TabletContexts);
+		explicit FStylusInputEventHandler(TArray<FDeprecatedStylusInputDevice>& TabletContexts);
 		virtual ~FStylusInputEventHandler() override = default;
 		virtual FString GetName() override { return "DeprecatedWindowsStylusInputInterfaceEventHandler"; }
-		virtual void OnPacket(const UE::StylusInput::FStylusInputPacket& Packet) override;
-		virtual void OnDebugEvent(const FString& Message) override;
+		virtual void OnPacket(const UE::StylusInput::FStylusInputPacket& Packet, UE::StylusInput::IStylusInputInstance* Instance) override;
+		virtual void OnDebugEvent(const FString& Message, UE::StylusInput::IStylusInputInstance* Instance) override;
 
 	private:
-		FDeprecatedStylusInputDevice* GetTabletContext(uint32 TabletContextId);
+		FDeprecatedStylusInputDevice* GetTabletContext(uint32 TabletContextId, UE::StylusInput::IStylusInputInstance* Instance);
 
-		UE::StylusInput::IStylusInputInstance* const Instance;
 		TArray<FDeprecatedStylusInputDevice>& TabletContexts;
 	};
 
diff --git a/Engine/Plugins/Editor/StylusInput/Source/StylusInput/Private/Windows/WindowsStylusInputInstance.cpp b/Engine/Plugins/Editor/StylusInput/Source/StylusInput/Private/Windows/WindowsStylusInputInstance.cpp
--- a/Engine/Plugins/Editor/StylusInput/Source/StylusInput/Private/Windows/WindowsStylusInputInstance.cpp
+++ b/Engine/Plugins/Editor/StylusInput/Source/StylusInput/Private/Windows/WindowsStylusInputInstance.cpp
@@ -372,6 +372,7 @@ namespace UE::StylusInput::Private::Windows
 		if (EventHandlerThread == EEventHandlerThread::OnGameThread)
 		{
 			AsyncPlugin = MakeUnique<FWindowsStylusInputPluginAsync>(
+				this,
 				FGetWindowContextCallback::CreateRaw(this, &FWindowsStylusInputInstance::GetWindowContext),
 				FUpdateTabletContextsCallback::CreateRaw(this, &FWindowsStylusInputInstance::UpdateTabletContexts),
 				EventHandler
@@ -391,6 +392,7 @@ namespace UE::StylusInput::Private::Windows
 		if (EventHandlerThread == EEventHandlerThread::Asynchronous)
 		{
 			SyncPlugin = MakeUnique<FWindowsStylusInputPluginSync>(
+				this,
 				FGetWindowContextCallback::CreateRaw(this, &FWindowsStylusInputInstance::GetWindowContext),
 				FUpdateTabletContextsCallback::CreateRaw(this, &FWindowsStylusInputInstance::UpdateTabletContexts),
 				EventHandler
diff --git a/Engine/Plugins/Editor/StylusInput/Source/StylusInput/Private/Windows/WindowsStylusInputPluginAsync.cpp b/Engine/Plugins/Editor/StylusInput/Source/StylusInput/Private/Windows/WindowsStylusInputPluginAsync.cpp
--- a/Engine/Plugins/Editor/StylusInput/Source/StylusInput/Private/Windows/WindowsStylusInputPluginAsync.cpp
+++ b/Engine/Plugins/Editor/StylusInput/Source/StylusInput/Private/Windows/WindowsStylusInputPluginAsync.cpp
@@ -6,10 +6,11 @@
 
 namespace UE::StylusInput::Private::Windows
 {
-	FWindowsStylusInputPluginAsync::FWindowsStylusInputPluginAsync(FGetWindowContextCallback&& GetWindowContext,
+	FWindowsStylusInputPluginAsync::FWindowsStylusInputPluginAsync(IStylusInputInstance* Instance,
+	                                                               FGetWindowContextCallback&& GetWindowContext,
 	                                                               FUpdateTabletContextsCallback&& UpdateTabletContextsCallback,
 	                                                               IStylusInputEventHandler* EventHandler)
-		: FWindowsStylusInputPluginBase(MoveTemp(GetWindowContext), MoveTemp(UpdateTabletContextsCallback))
+		: FWindowsStylusInputPluginBase(Instance, MoveTemp(GetWindowContext), MoveTemp(UpdateTabletContextsCallback))
 	{
 		if (EventHandler)
 		{
diff --git a/Engine/Plugins/Editor/StylusInput/Source/StylusInput/Private/Windows/WindowsStylusInputPluginAsync.h b/Engine/Plugins/Editor/StylusInput/Source/StylusInput/Private/Windows/WindowsStylusInputPluginAsync.h
--- a/Engine/Plugins/Editor/StylusInput/Source/StylusInput/Private/Windows/WindowsStylusInputPluginAsync.h
+++ b/Engine/Plugins/Editor/StylusInput/Source/StylusInput/Private/Windows/WindowsStylusInputPluginAsync.h
@@ -13,8 +13,8 @@ namespace UE::StylusInput::Private::Windows
 	class FWindowsStylusInputPluginAsync final : public IStylusAsyncPlugin, public FWindowsStylusInputPluginBase
 	{
 	public:
-		FWindowsStylusInputPluginAsync(FGetWindowContextCallback&& GetWindowContext, FUpdateTabletContextsCallback&& UpdateTabletContextsCallback,
-		                               IStylusInputEventHandler* EventHandler);
+		FWindowsStylusInputPluginAsync(IStylusInputInstance* Instance, FGetWindowContextCallback&& GetWindowContext,
+		                               FUpdateTabletContextsCallback&& UpdateTabletContextsCallback, IStylusInputEventHandler* EventHandler);
 		~FWindowsStylusInputPluginAsync();
 
 		// IUnknown
diff --git a/Engine/Plugins/Editor/StylusInput/Source/StylusInput/Private/Windows/WindowsStylusInputPluginBase.cpp b/Engine/Plugins/Editor/StylusInput/Source/StylusInput/Private/Windows/WindowsStylusInputPluginBase.cpp
--- a/Engine/Plugins/Editor/StylusInput/Source/StylusInput/Private/Windows/WindowsStylusInputPluginBase.cpp
+++ b/Engine/Plugins/Editor/StylusInput/Source/StylusInput/Private/Windows/WindowsStylusInputPluginBase.cpp
@@ -513,9 +513,11 @@ namespace UE::StylusInput::Private::Windows
 		return bWasRemoved;
 	}
 
-	FWindowsStylusInputPluginBase::FWindowsStylusInputPluginBase(FGetWindowContextCallback&& GetWindowContextCallback,
+	FWindowsStylusInputPluginBase::FWindowsStylusInputPluginBase(IStylusInputInstance* Instance,
+	                                                             FGetWindowContextCallback&& GetWindowContextCallback,
 	                                                             FUpdateTabletContextsCallback&& UpdateTabletContextsCallback)
-		: GetWindowContextCallback(MoveTemp(GetWindowContextCallback))
+		: Instance(Instance)
+		, GetWindowContextCallback(MoveTemp(GetWindowContextCallback))
 		, UpdateTabletContextsCallback(MoveTemp(UpdateTabletContextsCallback)) 
 	{
 		check(this->GetWindowContextCallback.IsBound());
@@ -526,7 +528,7 @@ namespace UE::StylusInput::Private::Windows
 	{
 		for (IStylusInputEventHandler* EventHandler : EventHandlers)
 		{
-			EventHandler->OnDebugEvent(Message);
+			EventHandler->OnDebugEvent(Message, Instance);
 		}
 	}
 
@@ -650,7 +652,7 @@ namespace UE::StylusInput::Private::Windows
 
 			for (IStylusInputEventHandler* EventHandler : EventHandlers)
 			{
-				EventHandler->OnPacket(Packet);
+				EventHandler->OnPacket(Packet, Instance);
 			}
 
 			PacketBuffer += PropertyCount;
diff --git a/Engine/Plugins/Editor/StylusInput/Source/StylusInput/Private/Windows/WindowsStylusInputPluginBase.h b/Engine/Plugins/Editor/StylusInput/Source/StylusInput/Private/Windows/WindowsStylusInputPluginBase.h
--- a/Engine/Plugins/Editor/StylusInput/Source/StylusInput/Private/Windows/WindowsStylusInputPluginBase.h
+++ b/Engine/Plugins/Editor/StylusInput/Source/StylusInput/Private/Windows/WindowsStylusInputPluginBase.h
@@ -20,6 +20,7 @@
 namespace UE::StylusInput
 {
 	class IStylusInputEventHandler;
+	class IStylusInputInstance;
 }
 
 namespace UE::StylusInput::Private::Windows
@@ -39,7 +40,8 @@ namespace UE::StylusInput::Private::Windows
 		float GetPacketsPerSecond() const { return PacketStats.GetPacketsPerSecond(); }
 
 	protected:
-		FWindowsStylusInputPluginBase(FGetWindowContextCallback&& GetWindowContextCallback, FUpdateTabletContextsCallback&& UpdateTabletContextsCallback);
+		FWindowsStylusInputPluginBase(IStylusInputInstance* Instance, FGetWindowContextCallback&& GetWindowContextCallback,
+		                              FUpdateTabletContextsCallback&& UpdateTabletContextsCallback);
 		~FWindowsStylusInputPluginBase() = default;
 
 		virtual FString GetName() const = 0;
@@ -58,13 +60,11 @@ namespace UE::StylusInput::Private::Windows
 		const FPacketProperty* GetPacketDescriptions(uint32 TabletContextId) const;
 		HRESULT UpdateTabletContexts(IRealTimeStylus* RealTimeStylus, uint32 TabletContextIDsNum, const TABLET_CONTEXT_ID* TabletContextIDs);
 
+		IStylusInputInstance *const Instance;
 		FPacketStats PacketStats;
-
 		FGetWindowContextCallback GetWindowContextCallback;
 		FUpdateTabletContextsCallback UpdateTabletContextsCallback;
-
 		FTabletContextContainer TabletContexts;
-
 		TArray<IStylusInputEventHandler*> EventHandlers;
 	};
 }
diff --git a/Engine/Plugins/Editor/StylusInput/Source/StylusInput/Private/Windows/WindowsStylusInputPluginSync.cpp b/Engine/Plugins/Editor/StylusInput/Source/StylusInput/Private/Windows/WindowsStylusInputPluginSync.cpp
--- a/Engine/Plugins/Editor/StylusInput/Source/StylusInput/Private/Windows/WindowsStylusInputPluginSync.cpp
+++ b/Engine/Plugins/Editor/StylusInput/Source/StylusInput/Private/Windows/WindowsStylusInputPluginSync.cpp
@@ -10,10 +10,11 @@
 
 namespace UE::StylusInput::Private::Windows
 {
-	FWindowsStylusInputPluginSync::FWindowsStylusInputPluginSync(FGetWindowContextCallback&& GetWindowContext,
+	FWindowsStylusInputPluginSync::FWindowsStylusInputPluginSync(IStylusInputInstance* Instance,
+	                                                             FGetWindowContextCallback&& GetWindowContext,
 	                                                             FUpdateTabletContextsCallback&& UpdateTabletContextsCallback,
 	                                                             IStylusInputEventHandler* EventHandler)
-		: FWindowsStylusInputPluginBase(MoveTemp(GetWindowContext), MoveTemp(UpdateTabletContextsCallback))
+		: FWindowsStylusInputPluginBase(Instance, MoveTemp(GetWindowContext), MoveTemp(UpdateTabletContextsCallback))
 	{
 		if (EventHandler)
 		{
diff --git a/Engine/Plugins/Editor/StylusInput/Source/StylusInput/Private/Windows/WindowsStylusInputPluginSync.h b/Engine/Plugins/Editor/StylusInput/Source/StylusInput/Private/Windows/WindowsStylusInputPluginSync.h
--- a/Engine/Plugins/Editor/StylusInput/Source/StylusInput/Private/Windows/WindowsStylusInputPluginSync.h
+++ b/Engine/Plugins/Editor/StylusInput/Source/StylusInput/Private/Windows/WindowsStylusInputPluginSync.h
@@ -13,8 +13,8 @@ namespace UE::StylusInput::Private::Windows
 	class FWindowsStylusInputPluginSync final : public IStylusSyncPlugin, public FWindowsStylusInputPluginBase
 	{
 	public:
-		FWindowsStylusInputPluginSync(FGetWindowContextCallback&& GetWindowContext, FUpdateTabletContextsCallback&& UpdateTabletContextsCallback,
-		                              IStylusInputEventHandler* EventHandler);
+		FWindowsStylusInputPluginSync(IStylusInputInstance* Instance, FGetWindowContextCallback&& GetWindowContext,
+		                              FUpdateTabletContextsCallback&& UpdateTabletContextsCallback, IStylusInputEventHandler* EventHandler);
 		~FWindowsStylusInputPluginSync();
 		
 		HRESULT CreateFreeThreadMarshaler();
diff --git a/Engine/Plugins/Editor/StylusInput/Source/StylusInput/Public/StylusInput.h b/Engine/Plugins/Editor/StylusInput/Source/StylusInput/Public/StylusInput.h
--- a/Engine/Plugins/Editor/StylusInput/Source/StylusInput/Public/StylusInput.h
+++ b/Engine/Plugins/Editor/StylusInput/Source/StylusInput/Public/StylusInput.h
@@ -127,8 +127,9 @@ namespace UE::StylusInput
 		 * instance. The implementation of this function is responsible for any thread synchronization that might be necessary. 
 		 *
 		 * @param Packet The packet being processed.
+		 * @param Instance The stylus input instance that sent the packet.
 		 */
-		virtual void OnPacket(const FStylusInputPacket& Packet) = 0;
+		virtual void OnPacket(const FStylusInputPacket& Packet, IStylusInputInstance* Instance) = 0;
 
 		/**
 		 * Callback for each debug event sent by the stylus input instance for diagnostic purposes.
@@ -137,7 +138,8 @@ namespace UE::StylusInput
 		 * The base class implementation of this function does nothing.
 		 *
 		 * @param Message The debug message sent by the stylus input instance.
+		 * @param Instance The stylus input instance that sent the debug message.
 		 */
-		virtual void OnDebugEvent(const FString& Message) {}
+		virtual void OnDebugEvent(const FString& Message, IStylusInputInstance* Instance) {}
 	};
 }
diff --git a/Engine/Plugins/Editor/StylusInput/Source/StylusInputDebugWidget/Private/StylusInputDebugWidget.cpp b/Engine/Plugins/Editor/StylusInput/Source/StylusInputDebugWidget/Private/StylusInputDebugWidget.cpp
--- a/Engine/Plugins/Editor/StylusInput/Source/StylusInputDebugWidget/Private/StylusInputDebugWidget.cpp
+++ b/Engine/Plugins/Editor/StylusInput/Source/StylusInputDebugWidget/Private/StylusInputDebugWidget.cpp
@@ -575,12 +575,12 @@ namespace UE::StylusInput::DebugWidget
 		check(this->OnDebugEventCallback.IsBound());
 	}
 
-	void FDebugEventHandlerAsynchronous::OnPacket(const FStylusInputPacket& Packet)
+	void FDebugEventHandlerAsynchronous::OnPacket(const FStylusInputPacket& Packet, IStylusInputInstance*)
 	{
 		PacketQueue.Enqueue(Packet);
 	}
 
-	void FDebugEventHandlerAsynchronous::OnDebugEvent(const FString& Message)
+	void FDebugEventHandlerAsynchronous::OnDebugEvent(const FString& Message, IStylusInputInstance*)
 	{
 		DebugEventQueue.Enqueue(Message);
 	}
@@ -607,12 +607,12 @@ namespace UE::StylusInput::DebugWidget
 		check(this->OnDebugEventCallback.IsBound());
 	}
 
-	void FDebugEventHandlerOnGameThread::OnPacket(const FStylusInputPacket& Packet)
+	void FDebugEventHandlerOnGameThread::OnPacket(const FStylusInputPacket& Packet, IStylusInputInstance*)
 	{
 		OnPacketCallback.Execute(Packet);
 	}
 
-	void FDebugEventHandlerOnGameThread::OnDebugEvent(const FString& Message)
+	void FDebugEventHandlerOnGameThread::OnDebugEvent(const FString& Message, IStylusInputInstance*)
 	{
 		OnDebugEventCallback.Execute(Message);
 	}
diff --git a/Engine/Plugins/Editor/StylusInput/Source/StylusInputDebugWidget/Private/StylusInputDebugWidget.h b/Engine/Plugins/Editor/StylusInput/Source/StylusInputDebugWidget/Private/StylusInputDebugWidget.h
--- a/Engine/Plugins/Editor/StylusInput/Source/StylusInputDebugWidget/Private/StylusInputDebugWidget.h
+++ b/Engine/Plugins/Editor/StylusInput/Source/StylusInputDebugWidget/Private/StylusInputDebugWidget.h
@@ -21,8 +21,8 @@ namespace UE::StylusInput::DebugWidget
 
 		virtual FString GetName() override { return "DebugEventHandlerAsynchronous"; }
 
-		virtual void OnPacket(const FStylusInputPacket& Packet) override;
-		virtual void OnDebugEvent(const FString& Message) override;
+		virtual void OnPacket(const FStylusInputPacket& Packet, IStylusInputInstance* Instance) override;
+		virtual void OnDebugEvent(const FString& Message, IStylusInputInstance* Instance) override;
 
 		virtual void Tick(float DeltaTime) override;
 		virtual TStatId GetStatId() const override { RETURN_QUICK_DECLARE_CYCLE_STAT(StylusInput_DebugEventHandlerAsynchronous, STATGROUP_Tickables); }
@@ -42,8 +42,8 @@ namespace UE::StylusInput::DebugWidget
 
 		virtual FString GetName() override { return "DebugEventHandlerOnGameThread"; }
 
-		virtual void OnPacket(const FStylusInputPacket& Packet) override;
-		virtual void OnDebugEvent(const FString& Message) override;
+		virtual void OnPacket(const FStylusInputPacket& Packet, IStylusInputInstance* Instance) override;
+		virtual void OnDebugEvent(const FString& Message, IStylusInputInstance* Instance) override;
 
 	private:
 		FOnPacketCallback OnPacketCallback;
