StylusInput: Rename menu entry for debug widget

diff --git a/Engine/Plugins/Editor/StylusInput/Source/StylusInputDebugWidget/Private/StylusInputDebugWidgetModule.cpp b/Engine/Plugins/Editor/StylusInput/Source/StylusInputDebugWidget/Private/StylusInputDebugWidgetModule.cpp
--- a/Engine/Plugins/Editor/StylusInput/Source/StylusInputDebugWidget/Private/StylusInputDebugWidgetModule.cpp
+++ b/Engine/Plugins/Editor/StylusInput/Source/StylusInputDebugWidget/Private/StylusInputDebugWidgetModule.cpp
@@ -54,7 +54,7 @@ namespace UE::StylusInput::DebugWidget
 
 		FTabSpawnerEntry& DebugWidgetSpawnerEntry = FGlobalTabmanager::Get()->RegisterNomadTabSpawner("StylusInputDebugWidget",
 			FOnSpawnTab::CreateRaw(this, &FStylusInputDebugWidgetModule::MakeDebugWidgetTab))
-				.SetDisplayName(LOCTEXT("DebugWidgetTitle", "Stylus Input Debug"))
+				.SetDisplayName(LOCTEXT("DebugWidgetTitle", "Stylus Input"))
 				.SetTooltipText(LOCTEXT("DebugWidgetTooltip", "Open a debug widget to verify stylus input event handling."))
 				.SetGroup(WorkspaceMenu::GetMenuStructure().GetDeveloperToolsDebugCategory())
 				.SetIcon(FSlateIcon(StyleSet->GetStyleSetName(), "StylusInput.Small"));
