StylusInput: Disable UStylusInputSubsystem tick

diff --git a/Engine/Plugins/Editor/StylusInput/Source/StylusInput/Private/DeprecatedStylusInputSubsystem.cpp b/Engine/Plugins/Editor/StylusInput/Source/StylusInput/Private/DeprecatedStylusInputSubsystem.cpp
--- a/Engine/Plugins/Editor/StylusInput/Source/StylusInput/Private/DeprecatedStylusInputSubsystem.cpp
+++ b/Engine/Plugins/Editor/StylusInput/Source/StylusInput/Private/DeprecatedStylusInputSubsystem.cpp
@@ -1,15 +1,25 @@
 ﻿// Copyright Epic Games, Inc. All Rights Reserved.
 
 #include <Framework/Docking/TabManager.h>
+#include <HAL/IConsoleManager.h>
 #include <Misc/App.h>
 
 #include "IStylusInputModule.h"
 
 PRAGMA_DISABLE_DEPRECATION_WARNINGS
 
-#define LOCTEXT_NAMESPACE "FStylusInputModule"
+#define LOCTEXT_NAMESPACE "StylusInputSubsystem"
 
-static const FName StylusInputDebugTabName = FName("StylusInputDebug");
+namespace UE::StylusInput::Private
+{
+	static bool bTickStylusInputSubsystem = false;
+	static FAutoConsoleVariableRef CVarEnableLegacySubsystem(
+		TEXT("stylusinput.EnableLegacySubsystem"),
+		bTickStylusInputSubsystem,
+		TEXT("Enable the legacy stylus input subsystem, which will automatically create a tablet input context for any window on mouse over. This subsystem is deprecated for UE 5.5, and will be removed entirely in UE 5.7."),
+		ECVF_Default
+	);
+}
 
 // This is the function that all platform-specific implementations are required to implement.
 TSharedPtr<IStylusInputInterfaceInternal> CreateStylusInputInterface();
@@ -34,7 +44,6 @@ void UStylusInputSubsystem::Initialize(FSubsystemCollectionBase& Collection)
 	if (!InputInterface.IsValid())
 	{
 		UE_LOG(LogStylusInput, Log, TEXT("StylusInput not supported on this platform."));
-		return;
 	}
 }
 
@@ -42,8 +51,6 @@ void UStylusInputSubsystem::Deinitialize()
 {
 	Super::Deinitialize();
 
-	FGlobalTabmanager::Get()->UnregisterTabSpawner(StylusInputDebugTabName);
-
 	InputInterface.Reset();
 
 	UE_LOG(LogStylusInput, Log, TEXT("Shutting down StylusInput subsystem."));
@@ -58,7 +65,7 @@ int32 UStylusInputSubsystem::NumInputDevices() const
 	return 0;
 }
 
-const IStylusInputDevice* UStylusInputSubsystem::GetInputDevice(int32 Index) const
+const IStylusInputDevice* UStylusInputSubsystem::GetInputDevice(const int32 Index) const
 {
 	if (InputInterface.IsValid())
 	{
@@ -67,14 +74,19 @@ const IStylusInputDevice* UStylusInputSubsystem::GetInputDevice(int32 Index) con
 	return nullptr;
 }
 
-void UStylusInputSubsystem::AddMessageHandler(IStylusMessageHandler& InHandler)
+void UStylusInputSubsystem::AddMessageHandler(IStylusMessageHandler& MessageHandler)
+{
+	MessageHandlers.AddUnique(&MessageHandler);
+}
+
+void UStylusInputSubsystem::RemoveMessageHandler(IStylusMessageHandler& MessageHandler)
 {
-	MessageHandlers.AddUnique(&InHandler);
+	MessageHandlers.Remove(&MessageHandler);
 }
 
-void UStylusInputSubsystem::RemoveMessageHandler(IStylusMessageHandler& InHandler)
+bool UStylusInputSubsystem::IsTickable() const
 {
-	MessageHandlers.Remove(&InHandler);
+	return UE::StylusInput::Private::bTickStylusInputSubsystem;
 }
 
 void UStylusInputSubsystem::Tick(float DeltaTime)
diff --git a/Engine/Plugins/Editor/StylusInput/Source/StylusInput/Public/IStylusInputModule.h b/Engine/Plugins/Editor/StylusInput/Source/StylusInput/Public/IStylusInputModule.h
--- a/Engine/Plugins/Editor/StylusInput/Source/StylusInput/Public/IStylusInputModule.h
+++ b/Engine/Plugins/Editor/StylusInput/Source/StylusInput/Public/IStylusInputModule.h
@@ -85,6 +85,7 @@ public:
 	PRAGMA_ENABLE_DEPRECATION_WARNINGS
 
 	// FTickableEditorObject implementation
+	virtual bool IsTickable() const override;
 	virtual void Tick(float DeltaTime) override;
 	virtual TStatId GetStatId() const override { RETURN_QUICK_DECLARE_CYCLE_STAT(UStylusInputSubsystem, STATGROUP_Tickables); }
 
