StylusInput: Update deprecation warning to include reference to new CVar

#rnx

diff --git a/Engine/Plugins/Editor/StylusInput/Source/StylusInput/Public/IStylusInputModule.h b/Engine/Plugins/Editor/StylusInput/Source/StylusInput/Public/IStylusInputModule.h
--- a/Engine/Plugins/Editor/StylusInput/Source/StylusInput/Public/IStylusInputModule.h
+++ b/Engine/Plugins/Editor/StylusInput/Source/StylusInput/Public/IStylusInputModule.h
@@ -56,7 +56,8 @@ public:
 
 
 UCLASS()
-class UE_DEPRECATED(5.5, "Please use the new API's window-specific IStylusInputInstance in StylusInput.h instead.") STYLUSINPUT_API UStylusInputSubsystem : 
+class UE_DEPRECATED(5.5, "Please use the new API's window-specific IStylusInputInstance in StylusInput.h instead. If you do need to use this subsystem, please also set the CVar 'stylusinput.EnableLegacySubsystem' to true, otherwise tablet input will not automatically be set up for each window.")
+STYLUSINPUT_API UStylusInputSubsystem : 
 	public UEditorSubsystem, 
 	public FTickableEditorObject
 {
