Remove #pragma once from cpp file

#rnx

diff --git a/Engine/Plugins/Editor/StylusInput/Source/StylusInput/Private/Windows/WindowsStylusInputPluginSync.cpp b/Engine/Plugins/Editor/StylusInput/Source/StylusInput/Private/Windows/WindowsStylusInputPluginSync.cpp
--- a/Engine/Plugins/Editor/StylusInput/Source/StylusInput/Private/Windows/WindowsStylusInputPluginSync.cpp
+++ b/Engine/Plugins/Editor/StylusInput/Source/StylusInput/Private/Windows/WindowsStylusInputPluginSync.cpp
@@ -1,7 +1,5 @@
 ﻿// Copyright Epic Games, Inc. All Rights Reserved.
 
-#pragma once
-
 #include "WindowsStylusInputPluginSync.h"
 
 #if PLATFORM_WINDOWS
diff --git a/Engine/Plugins/Editor/StylusInput/Source/StylusInput/Private/Windows/WindowsStylusInputTabletContext.cpp b/Engine/Plugins/Editor/StylusInput/Source/StylusInput/Private/Windows/WindowsStylusInputTabletContext.cpp
--- a/Engine/Plugins/Editor/StylusInput/Source/StylusInput/Private/Windows/WindowsStylusInputTabletContext.cpp
+++ b/Engine/Plugins/Editor/StylusInput/Source/StylusInput/Private/Windows/WindowsStylusInputTabletContext.cpp
@@ -1,7 +1,5 @@
 ﻿// Copyright Epic Games, Inc. All Rights Reserved.
 
-#pragma once
-
 #include "WindowsStylusInputTabletContext.h"
 
 #if PLATFORM_WINDOWS
