StylusInput: Minor bug fixes and cleanup

diff --git a/Engine/Plugins/Editor/StylusInput/Source/StylusInput/Private/StylusInputInterface.h b/Engine/Plugins/Editor/StylusInput/Source/StylusInput/Private/StylusInputInterface.h
--- a/Engine/Plugins/Editor/StylusInput/Source/StylusInput/Private/StylusInputInterface.h
+++ b/Engine/Plugins/Editor/StylusInput/Source/StylusInput/Private/StylusInputInterface.h
@@ -2,6 +2,8 @@
 
 #pragma once
 
+#include <HAL/Platform.h>
+
 class SWindow;
 
 namespace UE::StylusInput
diff --git a/Engine/Plugins/Editor/StylusInput/Source/StylusInput/Private/StylusInputUtils.h b/Engine/Plugins/Editor/StylusInput/Source/StylusInput/Private/StylusInputUtils.h
--- a/Engine/Plugins/Editor/StylusInput/Source/StylusInput/Private/StylusInputUtils.h
+++ b/Engine/Plugins/Editor/StylusInput/Source/StylusInput/Private/StylusInputUtils.h
@@ -3,6 +3,8 @@
 #pragma once
 
 #include <Logging/LogMacros.h>
+#include <Misc/ScopeRWLock.h>
+#include <Templates/SharedPointer.h>
 
 namespace UE::StylusInput::Private
 {
@@ -22,4 +24,83 @@ namespace UE::StylusInput::Private
 	{
 		UE_LOG(LogStylusInput, Verbose,  TEXT("%s: %s"), *Preamble, *Message);
 	}
+
+	template <typename DataType, bool ThreadSafe>
+	class TSharedRefDataContainer
+	{
+	public:
+		TSharedRef<DataType> Add(uint32 ID)
+		{
+			WriteScopeLockType WriteLock(RWLock);
+			return Data.Emplace_GetRef(MakeShared<DataType>(ID));
+		}
+
+		bool Contains(uint32 ID) const
+		{
+			ReadScopeLockType ReadLock(RWLock);
+			const TSharedRef<DataType>* TabletContext = Data.FindByPredicate([ID](const TSharedRef<DataType>& Context)
+			{
+				return Context->ID == ID;
+			});
+			return TabletContext != nullptr;
+		}
+
+		TSharedPtr<DataType> Get(uint32 ID) const
+		{
+			ReadScopeLockType ReadLock(RWLock);
+			const TSharedRef<DataType>* TabletContext = Data.FindByPredicate([ID](const TSharedRef<DataType>& Context)
+			{
+				return Context->ID == ID;
+			});
+			return TabletContext ? TabletContext->ToSharedPtr() : TSharedPtr<DataType>{};
+		}
+
+		bool Remove(uint32 ID)
+		{
+			int32 Index;
+			{
+				ReadScopeLockType ReadLock(RWLock);
+				Index = Data.IndexOfByPredicate([ID](const TSharedRef<DataType>& Tc) { return Tc->ID == ID; });
+				if (Index == INDEX_NONE)
+				{
+					return false;
+				}
+			}
+			WriteScopeLockType WriteLock(RWLock);
+			Data.RemoveAt(Index, EAllowShrinking::No);
+			return true;
+		}
+
+		void Clear()
+		{
+			WriteScopeLockType WriteLock(RWLock);
+			Data.Reset();
+		}
+
+		void Update(const TSharedRefDataContainer<DataType, false>& InTabletContexts)
+		{
+			WriteScopeLockType WriteLock(RWLock);
+			Data.Reset();
+			for (uint32 Index = 0, Num = InTabletContexts.Num(); Index < Num; ++Index)
+			{
+				Data.Emplace(InTabletContexts[Index]);
+			}
+		}
+
+		uint32 Num() const { return Data.Num(); }
+
+		const TSharedRef<DataType>& operator[](const uint32 Index) const { return Data[Index]; }
+
+	private:
+
+		struct FNoLock {};
+		struct FScopeNoLock { explicit FScopeNoLock(FNoLock&) {} };
+
+		using LockType = std::conditional_t<ThreadSafe, FRWLock, FNoLock>;
+		using ReadScopeLockType = std::conditional_t<ThreadSafe, FReadScopeLock, FScopeNoLock>;
+		using WriteScopeLockType = std::conditional_t<ThreadSafe, FWriteScopeLock, FScopeNoLock>;
+
+		UE_NO_UNIQUE_ADDRESS mutable LockType RWLock;
+		TArray<TSharedRef<DataType>> Data;
+	};
 }
diff --git a/Engine/Plugins/Editor/StylusInput/Source/StylusInput/Private/Windows/WindowsStylusInputInstance.cpp b/Engine/Plugins/Editor/StylusInput/Source/StylusInput/Private/Windows/WindowsStylusInputInstance.cpp
--- a/Engine/Plugins/Editor/StylusInput/Source/StylusInput/Private/Windows/WindowsStylusInputInstance.cpp
+++ b/Engine/Plugins/Editor/StylusInput/Source/StylusInput/Private/Windows/WindowsStylusInputInstance.cpp
@@ -17,12 +17,14 @@
 #define ASYNC_PLUGIN_INDEX 0
 #define SYNC_PLUGIN_INDEX 0
 
+#define LOCTEXT_NAMESPACE "WindowsStylusInputInstance"
 #define LOG_PREAMBLE "WindowsStylusInputInstance"
 
 namespace UE::StylusInput::Private::Windows
 {
-	FWindowsStylusInputInstance::FWindowsStylusInputInstance(const HWND OSWindowHandle)
-		: WindowsAPI(FWindowsStylusInputPlatformAPI::GetInstance())
+	FWindowsStylusInputInstance::FWindowsStylusInputInstance(const uint32 ID, const HWND OSWindowHandle)
+		: ID(ID)
+		, WindowsAPI(FWindowsStylusInputPlatformAPI::GetInstance())
 	{
 		Init(OSWindowHandle);
 	}
@@ -61,7 +63,7 @@ namespace UE::StylusInput::Private::Windows
 
 			return AsyncPlugin->AddEventHandler(EventHandler);
 		}
-		
+
 		if (Thread == EEventHandlerThread::Asynchronous)
 		{
 			if (!SyncPlugin)
@@ -94,7 +96,7 @@ namespace UE::StylusInput::Private::Windows
 		}
 
 		bool bWasRemoved = false;
-		
+
 		if (AsyncPlugin && AsyncPlugin->RemoveEventHandler(EventHandler))
 		{
 			bWasRemoved = true;
@@ -117,7 +119,7 @@ namespace UE::StylusInput::Private::Windows
 
 		if (!bWasRemoved)
 		{
-			LogError(LOG_PREAMBLE, FString::Format(TEXT("Event handler '%s' does not exist."), {EventHandler->GetName()}));
+			LogError(LOG_PREAMBLE, FString::Format(TEXT("Event handler {0} does not exist."), {EventHandler->GetName()}));
 		}
 
 		return bWasRemoved;
@@ -243,6 +245,11 @@ namespace UE::StylusInput::Private::Windows
 		return -1.0f;
 	}
 
+	FText FWindowsStylusInputInstance::GetName()
+	{
+		return FText::Format(LOCTEXT("WindowsRealTimeStylus", "Windows RealTimeStylus #{0}"), ID);
+	}
+
 	void FWindowsStylusInputInstance::SetupWindowContext(const HWND HWindow)
 	{
 		RECT WindowRectangle;
@@ -280,7 +287,7 @@ namespace UE::StylusInput::Private::Windows
 			static_cast<double>(WindowsAPI.GetDeviceCaps(HWindowDeviceContext, HORZRES)),
 			static_cast<double>(WindowsAPI.GetDeviceCaps(HWindowDeviceContext, VERTRES))
 		};
-		
+
 		if (!ensure(DisplaySizeInPixels.X > 0.0 && DisplaySizeInPixels.Y > 0.0))
 		{
 			LogError(LOG_PREAMBLE, "Display size in pixels is invalid; failed to setup device context!");
@@ -291,7 +298,7 @@ namespace UE::StylusInput::Private::Windows
 			static_cast<double>(WindowsAPI.GetDeviceCaps(HWindowDeviceContext, LOGPIXELSX)),
 			static_cast<double>(WindowsAPI.GetDeviceCaps(HWindowDeviceContext, LOGPIXELSY))
 		};
-		
+
 		if (!ensure(DisplayPixelsPerLogicalInch.X > 0.0 && DisplayPixelsPerLogicalInch.Y > 0.0))
 		{
 			LogError(LOG_PREAMBLE, "Display pixels per logical inch is invalid; failed to setup device context!");
@@ -327,15 +334,15 @@ namespace UE::StylusInput::Private::Windows
 		TabletContexts.Update(InTabletContexts);
 	}
 
-	void FWindowsStylusInputInstance::Init(HWND HWindow)
+	void FWindowsStylusInputInstance::Init(const HWND OSWindowHandle)
 	{
 		check(!RealTimeStylus.IsValid());
-		
-		SetupWindowContext(HWindow);
+
+		SetupWindowContext(OSWindowHandle);
 
 		void* OutInstance = nullptr;
 		if (Failed(WindowsAPI.CoCreateInstance(CLSID_RealTimeStylus, nullptr, CLSCTX_ALL, IID_IRealTimeStylus, &OutInstance), LOG_PREAMBLE))
-		{			
+		{
 			LogError(LOG_PREAMBLE, "Could not create RealTimeStylus COM object instance!");
 		}
 		RealTimeStylus = static_cast<IRealTimeStylus*>(OutInstance);
@@ -350,8 +357,8 @@ namespace UE::StylusInput::Private::Windows
 			{
 				LogError(LOG_PREAMBLE, "Could not set desired packet description!");
 			}
-			
-			if (Failed(RealTimeStylus->put_HWND(reinterpret_cast<HANDLE_PTR>(HWindow)), LOG_PREAMBLE))
+
+			if (Failed(RealTimeStylus->put_HWND(reinterpret_cast<HANDLE_PTR>(OSWindowHandle)), LOG_PREAMBLE))
 			{
 				LogError(LOG_PREAMBLE, "Could not set window handle!");
 			}
@@ -458,5 +465,6 @@ namespace UE::StylusInput::Private::Windows
 }
 
 #undef LOG_PREAMBLE
+#undef LOCTEXT_NAMESPACE
 
 #endif
diff --git a/Engine/Plugins/Editor/StylusInput/Source/StylusInput/Private/Windows/WindowsStylusInputInstance.h b/Engine/Plugins/Editor/StylusInput/Source/StylusInput/Private/Windows/WindowsStylusInputInstance.h
--- a/Engine/Plugins/Editor/StylusInput/Source/StylusInput/Private/Windows/WindowsStylusInputInstance.h
+++ b/Engine/Plugins/Editor/StylusInput/Source/StylusInput/Private/Windows/WindowsStylusInputInstance.h
@@ -31,7 +31,7 @@ namespace UE::StylusInput::Private::Windows
 	class FWindowsStylusInputInstance : public IStylusInputInstance
 	{
 	public:
-		explicit FWindowsStylusInputInstance(HWND OSWindowHandle);
+		explicit FWindowsStylusInputInstance(uint32 ID, HWND OSWindowHandle);
 		virtual ~FWindowsStylusInputInstance() override;
 
 		virtual bool AddEventHandler(IStylusInputEventHandler* EventHandler, EEventHandlerThread Thread) override;
@@ -42,8 +42,10 @@ namespace UE::StylusInput::Private::Windows
 
 		virtual float GetPacketsPerSecond(EEventHandlerThread EventHandlerThread) const override;
 
+		virtual FText GetName() override;
+
 	private:
-		void Init(HWND HWindow);
+		void Init(HWND OSWindowHandle);
 
 		void EnablePlugin(EEventHandlerThread EventHandlerThread, IStylusInputEventHandler* EventHandler);
 		void DisablePlugin(EEventHandlerThread EventHandlerThread);
@@ -53,17 +55,16 @@ namespace UE::StylusInput::Private::Windows
 
 		void UpdateTabletContexts(const FTabletContextContainer& InTabletContexts);
 
-		TComPtr<IRealTimeStylus> RealTimeStylus;
-
+		const uint32 ID;
+		const FWindowsStylusInputPlatformAPI& WindowsAPI;
 		FWindowContext WindowContext;
 
+		TComPtr<IRealTimeStylus> RealTimeStylus;
 		TUniquePtr<FWindowsStylusInputPluginAsync> AsyncPlugin;
 		TUniquePtr<FWindowsStylusInputPluginSync> SyncPlugin;
 
 		FTabletContextThreadSafeContainer TabletContexts;
 		FStylusInfoThreadSafeContainer StylusInfos;
-
-		const FWindowsStylusInputPlatformAPI& WindowsAPI;
 	};
 }
 
diff --git a/Engine/Plugins/Editor/StylusInput/Source/StylusInput/Private/Windows/WindowsStylusInputInterface.cpp b/Engine/Plugins/Editor/StylusInput/Source/StylusInput/Private/Windows/WindowsStylusInputInterface.cpp
--- a/Engine/Plugins/Editor/StylusInput/Source/StylusInput/Private/Windows/WindowsStylusInputInterface.cpp
+++ b/Engine/Plugins/Editor/StylusInput/Source/StylusInput/Private/Windows/WindowsStylusInputInterface.cpp
@@ -40,7 +40,8 @@ namespace UE::StylusInput::Private::Windows
 			return nullptr;
 		}
 
-		FWindowsStylusInputInstance* NewInstance = Instances.Emplace(&Window, {MakeUnique<FWindowsStylusInputInstance>(OSWindowHandle), 1}).Instance.Get();
+		FWindowsStylusInputInstance* NewInstance = Instances.Emplace(
+			&Window, {MakeUnique<FWindowsStylusInputInstance>(NextInstanceID++, OSWindowHandle), 1}).Instance.Get();
 		if (!ensureMsgf(NewInstance, TEXT("WindowsStylusInputInterface: Failed to create stylus input instance.")))
 		{
 			Instances.Remove(&Window);
@@ -70,12 +71,12 @@ namespace UE::StylusInput::Private::Windows
 					// Delete if there are no references left
 					Instances.Remove(Entry.Key);
 				}
-					
+
 				return true;
 			}
 		}
 
-		ensureMsgf(false, TEXT("WindowsStylusInputInterface: Failed to find provied instance."));
+		ensureMsgf(false, TEXT("WindowsStylusInputInterface: Failed to find provided instance."));
 		return false;
 	}
 
diff --git a/Engine/Plugins/Editor/StylusInput/Source/StylusInput/Private/Windows/WindowsStylusInputInterface.h b/Engine/Plugins/Editor/StylusInput/Source/StylusInput/Private/Windows/WindowsStylusInputInterface.h
--- a/Engine/Plugins/Editor/StylusInput/Source/StylusInput/Private/Windows/WindowsStylusInputInterface.h
+++ b/Engine/Plugins/Editor/StylusInput/Source/StylusInput/Private/Windows/WindowsStylusInputInterface.h
@@ -38,6 +38,7 @@ namespace UE::StylusInput::Private::Windows
 		};
 
 		TMap<SWindow*, FRefCountedInstance> Instances;
+		uint32 NextInstanceID = 0;
 	};
 }
 
diff --git a/Engine/Plugins/Editor/StylusInput/Source/StylusInput/Private/Windows/WindowsStylusInputPlatformAPI.cpp b/Engine/Plugins/Editor/StylusInput/Source/StylusInput/Private/Windows/WindowsStylusInputPlatformAPI.cpp
--- a/Engine/Plugins/Editor/StylusInput/Source/StylusInput/Private/Windows/WindowsStylusInputPlatformAPI.cpp
+++ b/Engine/Plugins/Editor/StylusInput/Source/StylusInput/Private/Windows/WindowsStylusInputPlatformAPI.cpp
@@ -52,6 +52,10 @@ namespace UE::StylusInput::Private::Windows
 			ExportHandle = static_cast<FFuncType>(FuncPtr);
 		};
 
+		const TCHAR* Kernel32DllNAme = TEXT("Kernel32.dll");
+		void* Kernel32DllHandle = GetDllHandle(Kernel32DllNAme);
+		GetDllExport(Kernel32DllNAme, Kernel32DllHandle, TEXT("GetTickCount"), GetTickCount);
+
 		const TCHAR* User32DllName = TEXT("User32.dll");
 		void* User32DllHandle = GetDllHandle(User32DllName);
 		GetDllExport(User32DllName, User32DllHandle, TEXT("GetClientRect"), GetClientRect);
@@ -132,6 +136,7 @@ namespace UE::StylusInput::Private::Windows
 		bool bSuccess = true;
 
 		bSuccess &= bInitializedComLibrary;
+		bSuccess &= GetTickCount != nullptr;
 		bSuccess &= GetClientRect != nullptr;
 		bSuccess &= GetDC != nullptr;
 		bSuccess &= ReleaseDC != nullptr;
diff --git a/Engine/Plugins/Editor/StylusInput/Source/StylusInput/Private/Windows/WindowsStylusInputPlatformAPI.h b/Engine/Plugins/Editor/StylusInput/Source/StylusInput/Private/Windows/WindowsStylusInputPlatformAPI.h
--- a/Engine/Plugins/Editor/StylusInput/Source/StylusInput/Private/Windows/WindowsStylusInputPlatformAPI.h
+++ b/Engine/Plugins/Editor/StylusInput/Source/StylusInput/Private/Windows/WindowsStylusInputPlatformAPI.h
@@ -23,6 +23,9 @@ namespace UE::StylusInput::Private::Windows
 
 		bool SatisfiesRequirements() const;
 
+		using FFuncGetTickCount = DWORD();
+		FFuncGetTickCount* GetTickCount = nullptr;
+
 		using FFuncGetClientRect = BOOL(HWND, LPRECT);
 		FFuncGetClientRect* GetClientRect = nullptr;
 
@@ -66,7 +69,7 @@ namespace UE::StylusInput::Private::Windows
 		bool bInitializedComLibrary = false;
 		bool bHasRTSComDllHandle = false;
 		bool bHasInkObjDllHandle = false;
-		
+
 	};
 }
 
diff --git a/Engine/Plugins/Editor/StylusInput/Source/StylusInput/Private/Windows/WindowsStylusInputPluginBase.cpp b/Engine/Plugins/Editor/StylusInput/Source/StylusInput/Private/Windows/WindowsStylusInputPluginBase.cpp
--- a/Engine/Plugins/Editor/StylusInput/Source/StylusInput/Private/Windows/WindowsStylusInputPluginBase.cpp
+++ b/Engine/Plugins/Editor/StylusInput/Source/StylusInput/Private/Windows/WindowsStylusInputPluginBase.cpp
@@ -91,17 +91,17 @@ namespace UE::StylusInput::Private::Windows
 
 		Packet.X = Value / Scale / Maximum * WindowWidth;
 	}
-	static_assert(std::is_same_v<FuncSetProperty, decltype(SetPropertyX)>);
+	static_assert(std::is_same_v<FPacketPropertyHandler::FFuncSetProperty, decltype(SetPropertyX)>);
 
 	void SetPropertyY(FStylusInputPacket& Packet, const int32 Value, const int8* Data)
 	{
 		const double& Scale = reinterpret_cast<const double*>(Data)[0];
 		const double& Maximum = reinterpret_cast<const double*>(Data)[1];
 		const double& WindowHeight = reinterpret_cast<const double*>(Data)[2];
-		
+
 		Packet.Y = Value / Scale / Maximum * WindowHeight;
 	}
-	static_assert(std::is_same_v<FuncSetProperty, decltype(SetPropertyY)>);
+	static_assert(std::is_same_v<FPacketPropertyHandler::FFuncSetProperty, decltype(SetPropertyY)>);
 
 	void SetPropertyZ(FStylusInputPacket& Packet, const int32 Value, const int8* Data)
 	{
@@ -110,25 +110,30 @@ namespace UE::StylusInput::Private::Windows
 
 		Packet.Z = static_cast<float>(Value - Minimum) / (Maximum - Minimum);
 	}
-	static_assert(std::is_same_v<FuncSetProperty, decltype(SetPropertyZ)>);
-	
+	static_assert(std::is_same_v<FPacketPropertyHandler::FFuncSetProperty, decltype(SetPropertyZ)>);
+
 	void SetPropertyPacketStatus(FStylusInputPacket& Packet, const int32 Value, const int8*)
 	{
 		Packet.PenStatus = static_cast<EPenStatus>(Value);
 	}
-	static_assert(std::is_same_v<FuncSetProperty, decltype(SetPropertyPacketStatus)>);
+	static_assert(std::is_same_v<FPacketPropertyHandler::FFuncSetProperty, decltype(SetPropertyPacketStatus)>);
 
 	void SetPropertyTimerTick(FStylusInputPacket& Packet, const int32 Value, const int8* Data)
 	{
-		Packet.TimerTick = Value;
+#if 0
+		Packet.TimerTick = Value; // The value provided from the RealTimeStylus API seems to only be non-zero if the stylus is down.
+#else
+		FWindowsStylusInputPlatformAPI::FFuncGetTickCount* const FuncGetTickCount =	*reinterpret_cast<decltype(FuncGetTickCount)*>(Data);
+		Packet.TimerTick = FuncGetTickCount();
+#endif
 	}
-	static_assert(std::is_same_v<FuncSetProperty, decltype(SetPropertyTimerTick)>);
+	static_assert(std::is_same_v<FPacketPropertyHandler::FFuncSetProperty, decltype(SetPropertyTimerTick)>);
 
 	void SetPropertySerialNumber(FStylusInputPacket& Packet, const int32 Value, const int8* Data)
 	{
 		Packet.SerialNumber = Value;
 	}
-	static_assert(std::is_same_v<FuncSetProperty, decltype(SetPropertySerialNumber)>);
+	static_assert(std::is_same_v<FPacketPropertyHandler::FFuncSetProperty, decltype(SetPropertySerialNumber)>);
 
 	void SetPropertyNormalPressure(FStylusInputPacket& Packet, const int32 Value, const int8* Data)
 	{
@@ -137,7 +142,7 @@ namespace UE::StylusInput::Private::Windows
 
 		Packet.NormalPressure = static_cast<float>(Value - Minimum) / (Maximum - Minimum);
 	}
-	static_assert(std::is_same_v<FuncSetProperty, decltype(SetPropertyNormalPressure)>);
+	static_assert(std::is_same_v<FPacketPropertyHandler::FFuncSetProperty, decltype(SetPropertyNormalPressure)>);
 
 	void SetPropertyTangentPressure(FStylusInputPacket& Packet, const int32 Value, const int8* Data)
 	{
@@ -146,7 +151,7 @@ namespace UE::StylusInput::Private::Windows
 
 		Packet.TangentPressure = static_cast<float>(Value - Minimum) / (Maximum - Minimum);
 	}
-	static_assert(std::is_same_v<FuncSetProperty, decltype(SetPropertyTangentPressure)>);
+	static_assert(std::is_same_v<FPacketPropertyHandler::FFuncSetProperty, decltype(SetPropertyTangentPressure)>);
 
 	void SetPropertyButtonPressure(FStylusInputPacket& Packet, const int32 Value, const int8* Data)
 	{
@@ -155,7 +160,7 @@ namespace UE::StylusInput::Private::Windows
 
 		Packet.ButtonPressure = static_cast<float>(Value - Minimum) / (Maximum - Minimum);
 	}
-	static_assert(std::is_same_v<FuncSetProperty, decltype(SetPropertyButtonPressure)>);
+	static_assert(std::is_same_v<FPacketPropertyHandler::FFuncSetProperty, decltype(SetPropertyButtonPressure)>);
 
 	void SetPropertyXTiltOrientation(FStylusInputPacket& Packet, const int32 Value, const int8* Data)
 	{
@@ -163,7 +168,7 @@ namespace UE::StylusInput::Private::Windows
 
 		Packet.XTiltOrientation = Value / Resolution;
 	}
-	static_assert(std::is_same_v<FuncSetProperty, decltype(SetPropertyXTiltOrientation)>);
+	static_assert(std::is_same_v<FPacketPropertyHandler::FFuncSetProperty, decltype(SetPropertyXTiltOrientation)>);
 
 	void SetPropertyYTiltOrientation(FStylusInputPacket& Packet, const int32 Value, const int8* Data)
 	{
@@ -171,7 +176,7 @@ namespace UE::StylusInput::Private::Windows
 
 		Packet.YTiltOrientation = Value / Resolution;
 	}
-	static_assert(std::is_same_v<FuncSetProperty, decltype(SetPropertyYTiltOrientation)>);
+	static_assert(std::is_same_v<FPacketPropertyHandler::FFuncSetProperty, decltype(SetPropertyYTiltOrientation)>);
 
 	void SetPropertyAzimuthOrientation(FStylusInputPacket& Packet, const int32 Value, const int8* Data)
 	{
@@ -179,7 +184,7 @@ namespace UE::StylusInput::Private::Windows
 
 		Packet.AzimuthOrientation = Value / Resolution;
 	}
-	static_assert(std::is_same_v<FuncSetProperty, decltype(SetPropertyAzimuthOrientation)>);
+	static_assert(std::is_same_v<FPacketPropertyHandler::FFuncSetProperty, decltype(SetPropertyAzimuthOrientation)>);
 
 	void SetPropertyAltitudeOrientation(FStylusInputPacket& Packet, const int32 Value, const int8* Data)
 	{
@@ -187,7 +192,7 @@ namespace UE::StylusInput::Private::Windows
 
 		Packet.AltitudeOrientation = Value / Resolution;
 	}
-	static_assert(std::is_same_v<FuncSetProperty, decltype(SetPropertyAltitudeOrientation)>);
+	static_assert(std::is_same_v<FPacketPropertyHandler::FFuncSetProperty, decltype(SetPropertyAltitudeOrientation)>);
 
 	void SetPropertyTwistOrientation(FStylusInputPacket& Packet, const int32 Value, const int8* Data)
 	{
@@ -195,7 +200,7 @@ namespace UE::StylusInput::Private::Windows
 
 		Packet.TwistOrientation = Value / Resolution;
 	}
-	static_assert(std::is_same_v<FuncSetProperty, decltype(SetPropertyTwistOrientation)>);
+	static_assert(std::is_same_v<FPacketPropertyHandler::FFuncSetProperty, decltype(SetPropertyTwistOrientation)>);
 
 	void SetPropertyPitchRotation(FStylusInputPacket& Packet, const int32 Value, const int8* Data)
 	{
@@ -203,7 +208,7 @@ namespace UE::StylusInput::Private::Windows
 
 		Packet.PitchRotation = Value / Resolution;
 	}
-	static_assert(std::is_same_v<FuncSetProperty, decltype(SetPropertyPitchRotation)>);
+	static_assert(std::is_same_v<FPacketPropertyHandler::FFuncSetProperty, decltype(SetPropertyPitchRotation)>);
 
 	void SetPropertyRollRotation(FStylusInputPacket& Packet, const int32 Value, const int8* Data)
 	{
@@ -211,7 +216,7 @@ namespace UE::StylusInput::Private::Windows
 
 		Packet.RollRotation = Value / Resolution;
 	}
-	static_assert(std::is_same_v<FuncSetProperty, decltype(SetPropertyRollRotation)>);
+	static_assert(std::is_same_v<FPacketPropertyHandler::FFuncSetProperty, decltype(SetPropertyRollRotation)>);
 
 	void SetPropertyYawRotation(FStylusInputPacket& Packet, const int32 Value, const int8* Data)
 	{
@@ -219,7 +224,7 @@ namespace UE::StylusInput::Private::Windows
 
 		Packet.YawRotation = Value / Resolution;
 	}
-	static_assert(std::is_same_v<FuncSetProperty, decltype(SetPropertyYawRotation)>);
+	static_assert(std::is_same_v<FPacketPropertyHandler::FFuncSetProperty, decltype(SetPropertyYawRotation)>);
 
 	void SetPropertyWidth(FStylusInputPacket& Packet, const int32 Value, const int8* Data)
 	{
@@ -227,7 +232,7 @@ namespace UE::StylusInput::Private::Windows
 
 		Packet.Width = Value / Resolution;
 	}
-	static_assert(std::is_same_v<FuncSetProperty, decltype(SetPropertyWidth)>);
+	static_assert(std::is_same_v<FPacketPropertyHandler::FFuncSetProperty, decltype(SetPropertyWidth)>);
 
 	void SetPropertyHeight(FStylusInputPacket& Packet, const int32 Value, const int8* Data)
 	{
@@ -235,7 +240,7 @@ namespace UE::StylusInput::Private::Windows
 
 		Packet.Height = Value / Resolution;
 	}
-	static_assert(std::is_same_v<FuncSetProperty, decltype(SetPropertyHeight)>);
+	static_assert(std::is_same_v<FPacketPropertyHandler::FFuncSetProperty, decltype(SetPropertyHeight)>);
 
 	void SetPropertyFingerContactConfidence(FStylusInputPacket& Packet, const int32 Value, const int8* Data)
 	{
@@ -244,15 +249,15 @@ namespace UE::StylusInput::Private::Windows
 
 		Packet.FingerContactConfidence = static_cast<float>(Value - Minimum) / (Maximum - Minimum);
 	}
-	static_assert(std::is_same_v<FuncSetProperty, decltype(SetPropertyFingerContactConfidence)>);
+	static_assert(std::is_same_v<FPacketPropertyHandler::FFuncSetProperty, decltype(SetPropertyFingerContactConfidence)>);
 
 	void SetPropertyDeviceContactID(FStylusInputPacket& Packet, const int32 Value, const int8* Data)
 	{
 		Packet.DeviceContactID = Value;
 	}
-	static_assert(std::is_same_v<FuncSetProperty, decltype(SetPropertyDeviceContactID)>);
+	static_assert(std::is_same_v<FPacketPropertyHandler::FFuncSetProperty, decltype(SetPropertyDeviceContactID)>);
 
-	FuncSetProperty* GetSetPropertyFunc(const EPacketPropertyType Type)
+	FPacketPropertyHandler::FFuncSetProperty* GetSetPropertyFunc(const EPacketPropertyType Type)
 	{
 		switch (Type)
 		{
@@ -299,7 +304,7 @@ namespace UE::StylusInput::Private::Windows
 		case EPacketPropertyType::DeviceContactID:
 			return &SetPropertyDeviceContactID;
 		default:
-			checkSlow(false);
+			checkSlow(false)
 			return nullptr;
 		}
 	}
@@ -347,31 +352,33 @@ namespace UE::StylusInput::Private::Windows
 		case EPacketPropertyType::Height:
 			Add(Description.Resolution);
 			break;
-		case EPacketPropertyType::PacketStatus:
 		case EPacketPropertyType::TimerTick:
+			Add(FWindowsStylusInputPlatformAPI::GetInstance().GetTickCount);
+			break;
+		case EPacketPropertyType::PacketStatus:
 		case EPacketPropertyType::SerialNumber:
 		case EPacketPropertyType::DeviceContactID:
 			// do nothing
 			break;
 		default:
-			checkSlow(false);
+			checkSlow(false)
 		}
 
-		check(std::distance(Data, DataCurrent) <= DataBufferLength);
+		check(std::distance(Data, DataCurrent) <= DataBufferLength)
 	}
 
 	bool SetupTabletContextPacketDescriptionData(IRealTimeStylus* RealTimeStylus, const FWindowContext& WindowContext, FTabletContext& TabletContext)
 	{
 		const FWindowsStylusInputPlatformAPI& WindowsAPI = FWindowsStylusInputPlatformAPI::GetInstance();
 
-		auto InvalidatePacketDescriptions = [&PacketDescriptions = TabletContext.PacketDescriptions](const int32 StartingIndex)
+		auto InvalidatePacketDescriptions = [&PacketDescriptions = TabletContext.PacketProperties](const int32 StartingIndex)
 		{
 			for (int32 Index = StartingIndex; Index < static_cast<int32>(EPacketPropertyType::Num_Enumerators); ++Index)
 			{
 				PacketDescriptions[Index] = {};
 			}
 		};
-		
+
 		FLOAT InkToDeviceScaleX = 0.0f;
 		FLOAT InkToDeviceScaleY = 0.0f;
 		ULONG PropertiesCount = 0;
@@ -407,26 +414,28 @@ namespace UE::StylusInput::Private::Windows
 					{
 						const ETabletSupportedProperties SupportedProperty = static_cast<ETabletSupportedProperties>(
 							1 << static_cast<std::underlying_type_t<ETabletSupportedProperties>>(PropertyConstant.PacketPropertyType)
-						);						
-						TabletContext.SupportedProperties = TabletContext.SupportedProperties | SupportedProperty;							
-
-						FPacketProperty& PropertyDescription = TabletContext.PacketDescriptions[Index];
-
-						PropertyDescription.Type = static_cast<EPacketPropertyType>(PropertyConstant.PacketPropertyType);
-						PropertyDescription.Minimum = Property.PropertyMetrics.nLogicalMin;
-						PropertyDescription.Maximum = Property.PropertyMetrics.nLogicalMax;
-						PropertyDescription.MetricUnit = static_cast<ETabletPropertyMetricUnit>(Property.PropertyMetrics.Units);
-						PropertyDescription.Resolution = Property.PropertyMetrics.fResolution;
-						PropertyDescription.SetProperty = GetSetPropertyFunc(PropertyConstant.PacketPropertyType);
-						AssignSetPropertyData<FPacketProperty::SetPropertyDataBufferLength>(PropertyConstant.PacketPropertyType, PropertyDescription.SetPropertyData, WindowContext, PropertyDescription);
-
-						if (static_cast<int8>(PropertyDescription.MetricUnit) >= static_cast<int8>(ETabletPropertyMetricUnit::Num_Enumerators))
+						);
+						TabletContext.SupportedProperties = TabletContext.SupportedProperties | SupportedProperty;
+
+						FPacketProperty& PropertyProperty = TabletContext.PacketProperties[Index];
+						PropertyProperty.Type = static_cast<EPacketPropertyType>(PropertyConstant.PacketPropertyType);
+						PropertyProperty.Minimum = Property.PropertyMetrics.nLogicalMin;
+						PropertyProperty.Maximum = Property.PropertyMetrics.nLogicalMax;
+						PropertyProperty.MetricUnit = static_cast<ETabletPropertyMetricUnit>(Property.PropertyMetrics.Units);
+						PropertyProperty.Resolution = Property.PropertyMetrics.fResolution;
+
+						FPacketPropertyHandler& PropertyPropertyHandler = TabletContext.PacketPropertyHandlers[Index];
+						PropertyPropertyHandler.SetProperty = GetSetPropertyFunc(PropertyConstant.PacketPropertyType);
+						AssignSetPropertyData<FPacketPropertyHandler::SetPropertyDataBufferLength>(
+							PropertyConstant.PacketPropertyType, PropertyPropertyHandler.SetPropertyData, WindowContext, PropertyProperty);
+
+						if (static_cast<int8>(PropertyProperty.MetricUnit) >= static_cast<int8>(ETabletPropertyMetricUnit::Num_Enumerators))
 						{
-							PropertyDescription.MetricUnit = ETabletPropertyMetricUnit::Default;
+							PropertyProperty.MetricUnit = ETabletPropertyMetricUnit::Default;
 
 							LogWarning(LOG_PREAMBLE, FString::Format(
 									   TEXT("Encountered unknown metric unit value '{0}' while evaluating packet description data for TabletContext with ID {1}."),
-									   {static_cast<int8>(PropertyDescription.MetricUnit), TabletContext.ID}));
+									   {static_cast<int8>(PropertyProperty.MetricUnit), TabletContext.ID}));
 						}
 
 						bFoundProperty = true;
@@ -480,17 +489,17 @@ namespace UE::StylusInput::Private::Windows
 
 		return bSuccess;
 	}
-	
+
 	bool FWindowsStylusInputPluginBase::AddEventHandler(IStylusInputEventHandler* EventHandler)
 	{
-		check(EventHandler);
+		check(EventHandler)
 
 		if (EventHandlers.Contains(EventHandler))
 		{
 			LogWarning(LOG_PREAMBLE, FString::Format(TEXT("Event handler '{0}' already exists in {1} plugin."), {EventHandler->GetName(), GetName()}));
 			return false;
 		}
-		
+
 		EventHandlers.Add(EventHandler);
 
 		LogVerbose(LOG_PREAMBLE, FString::Format(TEXT("Event handler '{0}' was added to {1} plugin."), {EventHandler->GetName(), GetName()}));
@@ -498,10 +507,9 @@ namespace UE::StylusInput::Private::Windows
 		return true;
 	}
 
-	
 	bool FWindowsStylusInputPluginBase::RemoveEventHandler(IStylusInputEventHandler* EventHandler)
 	{
-		check(EventHandler);
+		check(EventHandler)
 
 		const bool bWasRemoved = EventHandlers.Remove(EventHandler) > 0;
 
@@ -509,7 +517,7 @@ namespace UE::StylusInput::Private::Windows
 		{
 			LogVerbose(LOG_PREAMBLE, FString::Format(TEXT("Event handler '{0}' was removed from {1} plugin."), {EventHandler->GetName(), GetName()}));
 		}
-		
+
 		return bWasRemoved;
 	}
 
@@ -520,8 +528,8 @@ namespace UE::StylusInput::Private::Windows
 		, GetWindowContextCallback(MoveTemp(GetWindowContextCallback))
 		, UpdateTabletContextsCallback(MoveTemp(UpdateTabletContextsCallback)) 
 	{
-		check(this->GetWindowContextCallback.IsBound());
-		check(this->UpdateTabletContextsCallback.IsBound());
+		check(this->GetWindowContextCallback.IsBound())
+		check(this->UpdateTabletContextsCallback.IsBound())
 	}
 
 	void FWindowsStylusInputPluginBase::DebugEvent(const FString& Message) const
@@ -622,13 +630,13 @@ namespace UE::StylusInput::Private::Windows
 	}
 
 	HRESULT FWindowsStylusInputPluginBase::ProcessPackets(const StylusInfo* StylusInfo, uint32 PacketCount, uint32 PacketBufferLength, EPacketType Type,
-	                                                 const int32* PacketBuffer)
+	                                                      const int32* PacketBuffer)
 	{
-		checkSlow(StylusInfo);
-		checkSlow(PacketBuffer);
+		checkSlow(StylusInfo)
+		checkSlow(PacketBuffer)
 
-		const FPacketProperty* Descriptions = GetPacketDescriptions(StylusInfo->tcid);
-		if (!Descriptions)
+		const FPacketPropertyHandler* Handlers = GetPacketPropertyHandlers(StylusInfo->tcid);
+		if (!Handlers)
 		{
 			return E_FAIL;
 		}
@@ -641,13 +649,12 @@ namespace UE::StylusInput::Private::Windows
 			PacketStats.NewPacket();
 
 			FStylusInputPacket Packet{StylusInfo->tcid, StylusInfo->cid, Type};
-			
+
 			for (uint32 PropertyIndex = 0; PropertyIndex < PropertyCount; ++PropertyIndex)
 			{
-				const FPacketProperty& Description = Descriptions[PropertyIndex];
-				checkSlow(Description.SetProperty);
-
-				Description.SetProperty(Packet, PacketBuffer[PropertyIndex], Description.SetPropertyData);
+				const FPacketPropertyHandler& Handler = Handlers[PropertyIndex];
+				checkSlow(Handler.SetProperty)
+				Handler.SetProperty(Packet, PacketBuffer[PropertyIndex], Handler.SetPropertyData);
 			}
 
 			for (IStylusInputEventHandler* EventHandler : EventHandlers)
@@ -657,7 +664,7 @@ namespace UE::StylusInput::Private::Windows
 
 			PacketBuffer += PropertyCount;
 		}
-	
+
 		return S_OK;
 	}
 
@@ -679,7 +686,7 @@ namespace UE::StylusInput::Private::Windows
 
 	HRESULT FWindowsStylusInputPluginBase::ProcessTabletAdded(IInkTablet* Tablet)
 	{
-		// After a TabletAdded event, Windows Ink will fire a RealTimeStylusDisabled event directly followed by a RealTimeStylusEnabled event.
+		// After a TabletAdded event, a RealTimeStylusDisabled event will fire directly followed by a RealTimeStylusEnabled event.
 		// We are using these two events instead to update the tablet contexts.
 
 		return E_NOTIMPL;
@@ -688,22 +695,22 @@ namespace UE::StylusInput::Private::Windows
 	HRESULT FWindowsStylusInputPluginBase::ProcessTabletRemoved(LONG TabletIndex)
 	{
 		// For simplicity, we don't remove a tablet context when this event is received.
-		// However, sincere there are no more packets coming through for the removed tablet there should be nothing that's continuing to access the outdated
+		// However, since there are no more packets coming through for the removed tablet there should be nothing that's continuing to access the outdated
 		// tablet context data.
 
 		return E_NOTIMPL;
 	}
 
-	const FPacketProperty* FWindowsStylusInputPluginBase::GetPacketDescriptions(uint32 TabletContextId) const
+	const FPacketPropertyHandler* FWindowsStylusInputPluginBase::GetPacketPropertyHandlers(uint32 TabletContextID) const
 	{
-		const TSharedPtr<FTabletContext> TabletContext = TabletContexts.Get(TabletContextId);
-		return TabletContext ? TabletContext->PacketDescriptions : nullptr;
+		const TSharedPtr<FTabletContext> TabletContext = TabletContexts.Get(TabletContextID);
+		return TabletContext ? TabletContext->PacketPropertyHandlers : nullptr;
 	}
 
 	HRESULT FWindowsStylusInputPluginBase::UpdateTabletContexts(IRealTimeStylus* RealTimeStylus, uint32 TabletContextIDsNum, const TABLET_CONTEXT_ID* TabletContextIDs)
 	{
 		const FWindowContext& WindowContext = GetWindowContextCallback.Execute();
-		
+
 		bool bSuccess = true;
 
 		// Remove outdated tablet contexts
diff --git a/Engine/Plugins/Editor/StylusInput/Source/StylusInput/Private/Windows/WindowsStylusInputPluginBase.h b/Engine/Plugins/Editor/StylusInput/Source/StylusInput/Private/Windows/WindowsStylusInputPluginBase.h
--- a/Engine/Plugins/Editor/StylusInput/Source/StylusInput/Private/Windows/WindowsStylusInputPluginBase.h
+++ b/Engine/Plugins/Editor/StylusInput/Source/StylusInput/Private/Windows/WindowsStylusInputPluginBase.h
@@ -57,7 +57,7 @@ namespace UE::StylusInput::Private::Windows
 		HRESULT ProcessTabletRemoved(LONG TabletIndex);
 
 	private:
-		const FPacketProperty* GetPacketDescriptions(uint32 TabletContextId) const;
+		const FPacketPropertyHandler* GetPacketPropertyHandlers(uint32 TabletContextID) const;
 		HRESULT UpdateTabletContexts(IRealTimeStylus* RealTimeStylus, uint32 TabletContextIDsNum, const TABLET_CONTEXT_ID* TabletContextIDs);
 
 		IStylusInputInstance *const Instance;
diff --git a/Engine/Plugins/Editor/StylusInput/Source/StylusInput/Private/Windows/WindowsStylusInputStats.h b/Engine/Plugins/Editor/StylusInput/Source/StylusInput/Private/Windows/WindowsStylusInputStats.h
--- a/Engine/Plugins/Editor/StylusInput/Source/StylusInput/Private/Windows/WindowsStylusInputStats.h
+++ b/Engine/Plugins/Editor/StylusInput/Source/StylusInput/Private/Windows/WindowsStylusInputStats.h
@@ -6,7 +6,7 @@
 
 #if PLATFORM_WINDOWS
 
-#include "Windows/WindowsPlatformTime.h"
+#include <Windows/WindowsPlatformTime.h>
 
 namespace UE::StylusInput::Private::Windows
 {
diff --git a/Engine/Plugins/Editor/StylusInput/Source/StylusInput/Private/Windows/WindowsStylusInputTabletContext.h b/Engine/Plugins/Editor/StylusInput/Source/StylusInput/Private/Windows/WindowsStylusInputTabletContext.h
--- a/Engine/Plugins/Editor/StylusInput/Source/StylusInput/Private/Windows/WindowsStylusInputTabletContext.h
+++ b/Engine/Plugins/Editor/StylusInput/Source/StylusInput/Private/Windows/WindowsStylusInputTabletContext.h
@@ -8,6 +8,7 @@
 
 #include <StylusInputPacket.h>
 #include <StylusInputTabletContext.h>
+#include <StylusInputUtils.h>
 #include <Math/IntRect.h>
 #include <Math/Vector.h>
 
@@ -25,7 +26,7 @@ namespace UE::StylusInput::Private::Windows
 		Y, // The y-coordinate in the tablet coordinate space. Each packet contains this property by default. The origin (0,0) of the tablet is the upper-left corner.
 		Z, // The z-coordinate or distance of the pen tip from the tablet surface. The TabletPropertyMetricUnit enumeration type determines the unit of measurement for this property.
 		PacketStatus, // Contains one or more of the following flag values: The cursor is touching the drawing surface (Value = 1). The cursor is inverted. For example, the eraser end of the pen is pointing toward the surface (Value = 2). Not used (Value = 4). The barrel button is pressed (Value = 8).
-		TimerTick, // The time the packet was generated.
+		TimerTick, // The time the packet was generated in milliseconds since system start.
 		SerialNumber, // The packet property for identifying the packet. This is the same value you use to retrieve the packet from the packet queue.
 		NormalPressure, // The pressure of the pen tip perpendicular to the tablet surface. The greater the pressure on the pen tip, the more ink that is drawn.
 		TangentPressure, // The pressure of the pen tip along the plane of the tablet surface.
@@ -62,8 +63,6 @@ namespace UE::StylusInput::Private::Windows
 		Invalid_Enumerator = INDEX_NONE // THIS IS NOT A VALID ENUMERATOR, BUT IT CAN BE USED TO INDICATE THAT A PACKET PROPERTY ENTRY IS INVALID.
 	};
 
-	using FuncSetProperty = void(FStylusInputPacket&, int32, const int8*);
-
 	struct FPacketProperty
 	{
 		EPacketPropertyType Type = EPacketPropertyType::Invalid_Enumerator;
@@ -71,7 +70,12 @@ namespace UE::StylusInput::Private::Windows
 		int32 Minimum = 0;
 		int32 Maximum = 0;
 		float Resolution = 0.0f;
-		FuncSetProperty* SetProperty = nullptr;
+	};
+
+	struct FPacketPropertyHandler
+	{
+		using FFuncSetProperty = void(FStylusInputPacket&, int32, const int8*);
+		FFuncSetProperty* SetProperty = nullptr;
 
 		static constexpr int32 SetPropertyDataBufferLength = 24;
 		int8 SetPropertyData[SetPropertyDataBufferLength];
@@ -112,13 +116,14 @@ namespace UE::StylusInput::Private::Windows
 		FIntRect InputRectangle;
 		ETabletSupportedProperties SupportedProperties = ETabletSupportedProperties::None;
 		ETabletHardwareCapabilities HardwareCapabilities = ETabletHardwareCapabilities::None;
-		FPacketProperty PacketDescriptions[static_cast<uint32>(EPacketPropertyType::Num_Enumerators)];
+		FPacketProperty PacketProperties[static_cast<uint32>(EPacketPropertyType::Num_Enumerators)] = {};
+		FPacketPropertyHandler PacketPropertyHandlers[static_cast<uint32>(EPacketPropertyType::Num_Enumerators)] = {};
 	};
 
 	using FTabletContextContainer = TSharedRefDataContainer<FTabletContext, false>;
 	using FTabletContextThreadSafeContainer = TSharedRefDataContainer<FTabletContext, true>;
 
-	class FStylusButton final : IStylusInputStylusButton
+	class FStylusButton final : public IStylusInputStylusButton
 	{
 	public:
 		virtual FString GetID() const override { return ID; }
diff --git a/Engine/Plugins/Editor/StylusInput/Source/StylusInput/Private/Windows/WindowsStylusInputUtils.h b/Engine/Plugins/Editor/StylusInput/Source/StylusInput/Private/Windows/WindowsStylusInputUtils.h
--- a/Engine/Plugins/Editor/StylusInput/Source/StylusInput/Private/Windows/WindowsStylusInputUtils.h
+++ b/Engine/Plugins/Editor/StylusInput/Source/StylusInput/Private/Windows/WindowsStylusInputUtils.h
@@ -33,85 +33,6 @@ namespace UE::StylusInput::Private::Windows
 		LogCOMError(LogPreamble, Result);
 		return true;
 	}
-
-	template <typename DataType, bool ThreadSafe>
-	class TSharedRefDataContainer
-	{
-	public:
-		TSharedRef<DataType> Add(uint32 ID)
-		{
-			WriteScopeLockType WriteLock(RWLock);
-			return Data.Emplace_GetRef(MakeShared<DataType>(ID));
-		}
-
-		bool Contains(uint32 ID) const
-		{
-			ReadScopeLockType ReadLock(RWLock);
-			const TSharedRef<DataType>* TabletContext = Data.FindByPredicate([ID](const TSharedRef<DataType>& Context)
-			{
-				return Context->ID == ID;
-			});
-			return TabletContext != nullptr;
-		}
-
-		TSharedPtr<DataType> Get(uint32 ID) const
-		{
-			ReadScopeLockType ReadLock(RWLock);
-			const TSharedRef<DataType>* TabletContext = Data.FindByPredicate([ID](const TSharedRef<DataType>& Context)
-			{
-				return Context->ID == ID;
-			});
-			return TabletContext ? TabletContext->ToSharedPtr() : TSharedPtr<DataType>{};
-		}
-
-		bool Remove(uint32 ID)
-		{
-			int32 Index;
-			{
-				ReadScopeLockType ReadLock(RWLock);
-				Index = Data.IndexOfByPredicate([ID](const TSharedRef<DataType>& Tc) { return Tc->ID == ID; });
-				if (Index == INDEX_NONE)
-				{
-					return false;
-				}
-			}
-			WriteScopeLockType WriteLock(RWLock);
-			Data.RemoveAt(Index, EAllowShrinking::No);
-			return true;
-		}
-
-		void Clear()
-		{
-			WriteScopeLockType WriteLock(RWLock);
-			Data.Reset();
-		}
-
-		void Update(const TSharedRefDataContainer<DataType, false>& InTabletContexts)
-		{
-			WriteScopeLockType WriteLock(RWLock);
-			Data.Reset();
-			for (uint32 Index = 0, Num = InTabletContexts.Num(); Index < Num; ++Index)
-			{
-				Data.Emplace(InTabletContexts[Index]);
-			}
-		}
-
-		uint32 Num() const { return Data.Num(); }
-
-		const TSharedRef<DataType>& operator[](const uint32 Index) const { return Data[Index]; }
-
-	private:
-
-		struct FNoLock {};
-		struct FScopeNoLock { explicit FScopeNoLock(FNoLock&) {} };
-
-		using LockType = std::conditional_t<ThreadSafe, FRWLock, FNoLock>;
-		using ReadScopeLockType = std::conditional_t<ThreadSafe, FReadScopeLock, FScopeNoLock>;
-		using WriteScopeLockType = std::conditional_t<ThreadSafe, FWriteScopeLock, FScopeNoLock>;
-		
-		UE_NO_UNIQUE_ADDRESS mutable LockType RWLock;
-		TArray<TSharedRef<DataType>> Data;
-	};
 }
 
 #endif
diff --git a/Engine/Plugins/Editor/StylusInput/Source/StylusInput/Public/StylusInput.h b/Engine/Plugins/Editor/StylusInput/Source/StylusInput/Public/StylusInput.h
--- a/Engine/Plugins/Editor/StylusInput/Source/StylusInput/Public/StylusInput.h
+++ b/Engine/Plugins/Editor/StylusInput/Source/StylusInput/Public/StylusInput.h
@@ -104,6 +104,13 @@ namespace UE::StylusInput
 		 * @returns Approximate number of stylus input packets processed per second, or -1 if there is no valid data available.
 		 */
 		virtual float GetPacketsPerSecond(EEventHandlerThread Thread) const { return -1.0f; }
+
+		/**
+		 * Provides the name of this instance for diagnostic purposes. 
+		 * 
+		 * @returns Name of this instance.
+		 */
+		virtual FText GetName() = 0;
 	};
 
 	/**
diff --git a/Engine/Plugins/Editor/StylusInput/Source/StylusInput/Public/StylusInputPacket.h b/Engine/Plugins/Editor/StylusInput/Source/StylusInput/Public/StylusInputPacket.h
--- a/Engine/Plugins/Editor/StylusInput/Source/StylusInput/Public/StylusInputPacket.h
+++ b/Engine/Plugins/Editor/StylusInput/Source/StylusInput/Public/StylusInputPacket.h
@@ -13,7 +13,7 @@ namespace UE::StylusInput
 	enum class EPenStatus : uint8
 	{
 		None				= 0b00000000,	///< There is no pen flag.
-		CursorIsTouching	= 0b00000001,	///< The pen cursor is touching the drawing surface. 
+		CursorIsTouching	= 0b00000001,	///< The pen cursor is touching the drawing surface.
 		CursorIsInverted	= 0b00000010,	///< The pen cursor is inverted, e.g. the eraser end of the pen is pointing toward the drawing surface.
 		NotUsed				= 0b00000100,	///< Not used.
 		BarrelButtonPressed = 0b00001000	///< The barrel button is pressed.
@@ -49,8 +49,8 @@ namespace UE::StylusInput
 		float X = 0.0f;
 		float Y = 0.0f;
 		float Z = 0.0f;
-		int32 TimerTick = 0;
-		int32 SerialNumber = 0;
+		uint32 TimerTick = 0;
+		uint32 SerialNumber = 0;
 		float NormalPressure = 0.0f;
 		float TangentPressure = 0.0f;
 		float ButtonPressure = 0.0f;
@@ -69,7 +69,7 @@ namespace UE::StylusInput
 	};
 
 	/* Bitwise operators for EPenStatus to be used as flags without additional casts. */
-	
+
 	inline EPenStatus operator|(EPenStatus A, EPenStatus B)
 	{
 		return static_cast<EPenStatus>(static_cast<std::underlying_type_t<EPenStatus>>(A) | static_cast<std::underlying_type_t<EPenStatus>>(B));
diff --git a/Engine/Plugins/Editor/StylusInput/Source/StylusInput/Public/StylusInputTabletContext.h b/Engine/Plugins/Editor/StylusInput/Source/StylusInput/Public/StylusInputTabletContext.h
--- a/Engine/Plugins/Editor/StylusInput/Source/StylusInput/Public/StylusInputTabletContext.h
+++ b/Engine/Plugins/Editor/StylusInput/Source/StylusInput/Public/StylusInputTabletContext.h
@@ -30,7 +30,7 @@ namespace UE::StylusInput
 		Y					    = 1 <<  1, ///< The y-coordinate in the coordinate space of the window that was used to create the stylus input instance. Each packet contains this property by default.
 		Z					    = 1 <<  2, ///< The z-coordinate or distance of the pen tip from the tablet surface. The TabletPropertyMetricUnit enumeration  type determines the unit of measurement for this property.
 		PacketStatus		    = 1 <<  3, ///< Contains one or more of the following flag values: The cursor is touching the drawing surface (Value = 1). The cursor is inverted. For example, the eraser end of the pen is pointing toward the surface (Value = 2). Not used (Value = 4). The barrel button is pressed (Value = 8).
-		TimerTick			    = 1 <<  4, ///< The time the packet was generated.
+		TimerTick			    = 1 <<  4, ///< The time the packet was generated in milliseconds since system start.
 		SerialNumber		    = 1 <<  5, ///< The packet property for identifying the packet. This is the same value you use to retrieve the packet from the packet queue.
 		NormalPressure		    = 1 <<  6, ///< The pressure of the pen tip perpendicular to the tablet surface. The greater the pressure on the pen tip, the  more ink that is drawn.
 		TangentPressure		    = 1 <<  7, ///< The pressure of the pen tip along the plane of the tablet surface.
diff --git a/Engine/Plugins/Editor/StylusInput/Source/StylusInput/StylusInput.Build.cs b/Engine/Plugins/Editor/StylusInput/Source/StylusInput/StylusInput.Build.cs
--- a/Engine/Plugins/Editor/StylusInput/Source/StylusInput/StylusInput.Build.cs
+++ b/Engine/Plugins/Editor/StylusInput/Source/StylusInput/StylusInput.Build.cs
@@ -7,22 +7,20 @@ namespace UnrealBuildTool.Rules
 		public StylusInput(ReadOnlyTargetRules Target) : base(Target)
 		{
 			PublicDependencyModuleNames.AddRange(
-				new string[]
-				{
+				[
 					"Core",
 					"CoreUObject",
 					"EditorSubsystem",
 					"Engine",
 					"Slate",
 					"UnrealEd",
-				}
+				]
 			);
 
 			PrivateDependencyModuleNames.AddRange(
-				new string[]
-				{
+				[
 					"SlateCore",
-				}
+				]
 			);
 		}
 	}
diff --git a/Engine/Plugins/Editor/StylusInput/Source/StylusInputDebugWidget/Private/StylusInputDebugPaintWidget.cpp b/Engine/Plugins/Editor/StylusInput/Source/StylusInputDebugWidget/Private/StylusInputDebugPaintWidget.cpp
--- a/Engine/Plugins/Editor/StylusInput/Source/StylusInputDebugWidget/Private/StylusInputDebugPaintWidget.cpp
+++ b/Engine/Plugins/Editor/StylusInput/Source/StylusInputDebugWidget/Private/StylusInputDebugPaintWidget.cpp
@@ -50,6 +50,11 @@ namespace UE::StylusInput::DebugWidget
 		return LayerId;
 	}
 
+	TOptional<EMouseCursor::Type> SStylusInputDebugPaintWidget::GetCursor() const
+	{
+		return EMouseCursor::Crosshairs;
+	}
+
 	int32 SStylusInputDebugPaintWidget::DrawPaintPackets(const FGeometry& AllottedGeometry, FSlateWindowElementList& OutDrawElements, int32 LayerId) const
 	{
 		if (PaintPackets.IsEmpty())
diff --git a/Engine/Plugins/Editor/StylusInput/Source/StylusInputDebugWidget/Private/StylusInputDebugPaintWidget.h b/Engine/Plugins/Editor/StylusInput/Source/StylusInputDebugWidget/Private/StylusInputDebugPaintWidget.h
--- a/Engine/Plugins/Editor/StylusInput/Source/StylusInputDebugWidget/Private/StylusInputDebugPaintWidget.h
+++ b/Engine/Plugins/Editor/StylusInput/Source/StylusInputDebugWidget/Private/StylusInputDebugPaintWidget.h
@@ -29,6 +29,8 @@ namespace UE::StylusInput::DebugWidget
 		virtual int32 OnPaint(const FPaintArgs& Args, const FGeometry& AllottedGeometry, const FSlateRect& MyCullingRect,
 							  FSlateWindowElementList& OutDrawElements, int32 LayerId, const FWidgetStyle& InWidgetStyle, bool bParentEnabled) const override;
 
+		virtual TOptional<EMouseCursor::Type> GetCursor() const;
+
 	private:
 		struct FPaintPacket
 		{
diff --git a/Engine/Plugins/Editor/StylusInput/Source/StylusInputDebugWidget/Private/StylusInputDebugWidget.cpp b/Engine/Plugins/Editor/StylusInput/Source/StylusInputDebugWidget/Private/StylusInputDebugWidget.cpp
--- a/Engine/Plugins/Editor/StylusInput/Source/StylusInputDebugWidget/Private/StylusInputDebugWidget.cpp
+++ b/Engine/Plugins/Editor/StylusInput/Source/StylusInputDebugWidget/Private/StylusInputDebugWidget.cpp
@@ -52,7 +52,7 @@ namespace UE::StylusInput::DebugWidget
 
 		TSharedPtr<SVerticalBox> TopLeftOverlay;
 		TSharedPtr<SVerticalBox> BottomLeftOverlay;
-		
+
 		ChildSlot
 			.VAlign(VAlign_Fill)
 			.HAlign(HAlign_Fill)
@@ -165,24 +165,34 @@ namespace UE::StylusInput::DebugWidget
 		};
 
 		AddToOverlay(TopLeftOverlay,
-		             [&TabletContext = LastPacketData.TabletContext]
-		             {
-			             return TabletContext
-				                    ? FText::Format(
-					                    LOCTEXT("TabletContextID", "Tablet Context ID: {0}"),
-					                    FText::FromString(FString::Printf(TEXT("%x"), TabletContext->GetID())))
-				                    : FText();
-		             });
+					 [&StylusInput = StylusInput]
+					 {
+						 return StylusInput
+							        ? FText::Format(LOCTEXT("Instance", "Instance: {0}"), StylusInput->GetName())
+							        : FText();
+					 });
 
 		AddToOverlay(TopLeftOverlay,
-		             [&TabletContext = LastPacketData.TabletContext]
-		             {
-			             return TabletContext
-				                    ? FText::Format(
-					                    LOCTEXT("TabletContextName", "Name: {0}"),
-					                    FText::FromString(TabletContext->GetName()))
-				                    : FText();
-		             });
+					 [&TabletContext = LastPacketData.TabletContext]
+					 {
+						 return TabletContext
+							        ? FText::Format(LOCTEXT("TabletID", "Tablet ID: {0}"), TabletContext->GetID())
+							        : FText();
+					 });
+
+		AddToOverlay(TopLeftOverlay,
+					 [&TabletContext = LastPacketData.TabletContext]
+					 {
+						 return TabletContext
+									? FText::Format(
+										LOCTEXT("TabletName", "Tablet Name: {0}"),
+										FText::FromString([&TabletContext]
+										{
+											const FString Name = TabletContext->GetName();
+											return !Name.IsEmpty() ? Name : "<empty>";
+										}()))
+									: FText();
+					 });
 
 		AddToOverlay(TopLeftOverlay,
 		             [&TabletContext = LastPacketData.TabletContext]
@@ -191,7 +201,7 @@ namespace UE::StylusInput::DebugWidget
 			             {
 				             const FIntRect InputRectangle = TabletContext->GetInputRectangle();
 				             return FText::Format(
-					             LOCTEXT("TabletContextInputRectangle", "Input Rectangle: ({0}, {1}) x ({2}, {3})"),
+					             LOCTEXT("TabletInputRectangle", "Tablet Input Rectangle: ({0}, {1}) x ({2}, {3})"),
 					             InputRectangle.Min.X, InputRectangle.Min.Y, InputRectangle.Max.X, InputRectangle.Max.Y);
 			             }
 			             return FText();
@@ -217,15 +227,15 @@ namespace UE::StylusInput::DebugWidget
 					             }
 					             if ((Capabilities & ETabletHardwareCapabilities::CursorMustTouch) != ETabletHardwareCapabilities::None)
 					             {
-						             CapabilitiesStr += CapabilitiesStr.IsEmpty() ? "CursorMustTouch" : " & CursorMustTouch";
+						             CapabilitiesStr += CapabilitiesStr.IsEmpty() ? "CursorMustTouch" : ", CursorMustTouch";
 					             }
 					             if ((Capabilities & ETabletHardwareCapabilities::HardProximity) != ETabletHardwareCapabilities::None)
 					             {
-						             CapabilitiesStr += CapabilitiesStr.IsEmpty() ? "HardProximity" : " & HardProximity";
+						             CapabilitiesStr += CapabilitiesStr.IsEmpty() ? "HardProximity" : ", HardProximity";
 					             }
 					             if ((Capabilities & ETabletHardwareCapabilities::CursorsHavePhysicalIds) != ETabletHardwareCapabilities::None)
 					             {
-						             CapabilitiesStr += CapabilitiesStr.IsEmpty() ? "CursorsHavePhysicalIds" : " & CursorsHavePhysicalIds";
+						             CapabilitiesStr += CapabilitiesStr.IsEmpty() ? "CursorsHavePhysicalIds" : ", CursorsHavePhysicalIds";
 					             }
 				             }
 
@@ -238,7 +248,7 @@ namespace UE::StylusInput::DebugWidget
 				             }
 
 				             return FText::Format(
-					             LOCTEXT("TabletContextHardwareCapabilities", "Hardware Capabilities: {0}"), FText::FromString(CapabilitiesStr));
+					             LOCTEXT("TabletHardwareCapabilities", "Tablet Hardware Capabilities: {0}"), FText::FromString(CapabilitiesStr));
 			             }
 			             return FText();
 		             });
@@ -247,9 +257,7 @@ namespace UE::StylusInput::DebugWidget
 		             [&StylusInfo = LastPacketData.StylusInfo]
 		             {
 			             return StylusInfo
-				                    ? FText::Format(
-					                    LOCTEXT("StylusID", "Stylus ID: {0}"),
-					                    FText::FromString(FString::Printf(TEXT("%d"), StylusInfo->GetID())))
+				                    ? FText::Format(LOCTEXT("StylusID", "Stylus ID: {0}"), StylusInfo->GetID())
 				                    : FText();
 		             });
 
@@ -259,10 +267,14 @@ namespace UE::StylusInput::DebugWidget
 			             return StylusInfo
 				                    ? FText::Format(
 					                    LOCTEXT("StylusName", "Stylus Name: {0}"),
-					                    FText::FromString(StylusInfo->GetName()))
+					                    FText::FromString([&StylusInfo]
+					                    {
+						                    const FString Name = StylusInfo->GetName();
+						                    return !Name.IsEmpty() ? Name : "<empty>";
+					                    }()))
 				                    : FText();
 		             });
-		
+
 		AddToOverlay(TopLeftOverlay,
 		             [&StylusInfo = LastPacketData.StylusInfo]
 		             {
@@ -273,13 +285,14 @@ namespace UE::StylusInput::DebugWidget
 				             {
 					             if (const auto* Button = StylusInfo->GetButton(Index))
 					             {
+						             const FString Name = Button->GetName();
 						             if (Index == 0)
 						             {
-							             ButtonsString = Button->GetName();
+							             ButtonsString = !Name.IsEmpty() ? Name : "<unknown>";
 						             }
 						             else
 						             {
-							             ButtonsString += ", " + Button->GetName();
+							             ButtonsString += ", " + (!Name.IsEmpty() ? Name : "<unknown>");
 						             }
 					             }
 				             }
@@ -393,31 +406,31 @@ namespace UE::StylusInput::DebugWidget
 		AddToOverlayWithVisibility(BottomLeftOverlay, ETabletSupportedProperties::XTiltOrientation,
 		                           [&XTiltOrientation = LastPacketData.Packet.XTiltOrientation]
 		                           {
-			                           return FText::Format(LOCTEXT("XTiltOrientation", "X Tilt Orientation: {0}"), XTiltOrientation);
+			                           return FText::Format(LOCTEXT("XTiltOrientation", "X Tilt Orientation: {0}\u00B0"), XTiltOrientation);
 		                           });
 
 		AddToOverlayWithVisibility(BottomLeftOverlay, ETabletSupportedProperties::YTiltOrientation,
 		                           [&YTiltOrientation = LastPacketData.Packet.YTiltOrientation]
 		                           {
-			                           return FText::Format(LOCTEXT("YTiltOrientation", "Y Tilt Orientation: {0}"), YTiltOrientation);
+			                           return FText::Format(LOCTEXT("YTiltOrientation", "Y Tilt Orientation: {0}\u00B0"), YTiltOrientation);
 		                           });
 
 		AddToOverlayWithVisibility(BottomLeftOverlay, ETabletSupportedProperties::AzimuthOrientation,
 		                           [&AzimuthOrientation = LastPacketData.Packet.AzimuthOrientation]
 		                           {
-			                           return FText::Format(LOCTEXT("AzimuthOrientation", "Azimuth Orientation: {0}"), AzimuthOrientation);
+			                           return FText::Format(LOCTEXT("AzimuthOrientation", "Azimuth Orientation: {0}\u00B0"), AzimuthOrientation);
 		                           });
 
 		AddToOverlayWithVisibility(BottomLeftOverlay, ETabletSupportedProperties::AltitudeOrientation,
 		                           [&AltitudeOrientation = LastPacketData.Packet.AltitudeOrientation]
 		                           {
-			                           return FText::Format(LOCTEXT("AltitudeOrientation", "Altitude Orientation: {0}"), AltitudeOrientation);
+			                           return FText::Format(LOCTEXT("AltitudeOrientation", "Altitude Orientation: {0}\u00B0"), AltitudeOrientation);
 		                           });
 
 		AddToOverlayWithVisibility(BottomLeftOverlay, ETabletSupportedProperties::TwistOrientation,
 		                           [&TwistOrientation = LastPacketData.Packet.TwistOrientation]
 		                           {
-			                           return FText::Format(LOCTEXT("TwistOrientation", "Twist Orientation: {0}"), TwistOrientation);
+			                           return FText::Format(LOCTEXT("TwistOrientation", "Twist Orientation: {0}\u00B0"), TwistOrientation);
 		                           });
 
 		AddToOverlayWithVisibility(BottomLeftOverlay, ETabletSupportedProperties::PitchRotation,
@@ -734,7 +747,7 @@ namespace UE::StylusInput::DebugWidget
 		const TSharedPtr<IStylusInputStylusInfo>* StylusInfo = StylusInfos.Find(StylusID);
 		if (!StylusInfo)
 		{
-			if (const TSharedPtr<IStylusInputStylusInfo>& NewStylusInfo = StylusInput->GetStylusInfo(StylusID))
+			if (const TSharedPtr<IStylusInputStylusInfo> NewStylusInfo = StylusInput->GetStylusInfo(StylusID))
 			{
 				StylusInfo = &StylusInfos.Emplace(StylusID, NewStylusInfo);
 			}
