StylusInput: Include plug and play ID in tablet context name

#rnx

diff --git a/Engine/Plugins/Editor/StylusInput/Source/StylusInput/Private/Windows/WindowsStylusInputTabletContext.h b/Engine/Plugins/Editor/StylusInput/Source/StylusInput/Private/Windows/WindowsStylusInputTabletContext.h
--- a/Engine/Plugins/Editor/StylusInput/Source/StylusInput/Private/Windows/WindowsStylusInputTabletContext.h
+++ b/Engine/Plugins/Editor/StylusInput/Source/StylusInput/Private/Windows/WindowsStylusInputTabletContext.h
@@ -105,7 +105,16 @@ namespace UE::StylusInput::Private::Windows
 		{}
 
 		virtual uint32 GetID() const override { return ID; }
-		virtual FString GetName() const override { return Name; }
+
+		virtual FString GetName() const override
+		{
+			if (!Name.IsEmpty() && !PlugAndPlayID.IsEmpty())
+			{
+				return FString::Format(TEXT("{0} ({1})"), {Name, PlugAndPlayID});
+			}
+			return !Name.IsEmpty() ? Name : PlugAndPlayID;
+		}
+
 		virtual FIntRect GetInputRectangle() const override { return InputRectangle; }
 		virtual ETabletHardwareCapabilities GetHardwareCapabilities() const override { return HardwareCapabilities; }
 		virtual ETabletSupportedProperties GetSupportedProperties() const override { return SupportedProperties; }
