StylusInput: Rename and move Windows implementation to "RealTimeStylus"

#rnx

diff --git a/Engine/Plugins/Editor/StylusInput/Source/StylusInput/Private/StylusInput.cpp b/Engine/Plugins/Editor/StylusInput/Source/StylusInput/Private/StylusInput.cpp
--- a/Engine/Plugins/Editor/StylusInput/Source/StylusInput/Private/StylusInput.cpp
+++ b/Engine/Plugins/Editor/StylusInput/Source/StylusInput/Private/StylusInput.cpp
@@ -6,7 +6,7 @@
 #include "StylusInputUtils.h"
 
 #if PLATFORM_WINDOWS
-	#include <Windows/WindowsStylusInputInterface.h>
+	#include <Windows/RealTimeStylus/RealTimeStylusInterface.h>
 #endif
 
 namespace UE::StylusInput
@@ -60,7 +60,7 @@ namespace UE::StylusInput
 		private:
 			using FPlatformInterface =
 #if PLATFORM_WINDOWS
-				Windows::FWindowsStylusInputInterface
+				RealTimeStylus::FRealTimeStylusInterface
 #else
 				FNotSupportedStylusInputInterface
 #endif
diff --git a/Engine/Plugins/Editor/StylusInput/Source/StylusInput/Private/Windows/RealTimeStylus/RealTimeStylusAPI.cpp b/Engine/Plugins/Editor/StylusInput/Source/StylusInput/Private/Windows/RealTimeStylus/RealTimeStylusAPI.cpp
new file mode 100644
--- /dev/null
+++ b/Engine/Plugins/Editor/StylusInput/Source/StylusInput/Private/Windows/RealTimeStylus/RealTimeStylusAPI.cpp
@@ -0,0 +1,160 @@
+﻿// Copyright Epic Games, Inc. All Rights Reserved.
+
+#include "RealTimeStylusAPI.h"
+
+#if PLATFORM_WINDOWS
+
+#include <KnownFolders.h>
+#include <shtypes.h>
+#include <StylusInputUtils.h>
+#include <Windows/WindowsPlatformProcess.h>
+
+#include "RealTimeStylusUtils.h"
+
+#define LOG_PREAMBLE "RealTimeStylusAPI"
+
+namespace UE::StylusInput::Private::RealTimeStylus
+{
+	FRealTimeStylusAPI::FRealTimeStylusAPI()
+	{
+		bInitializedComLibrary = FWindowsPlatformMisc::CoInitialize();
+		if (!bInitializedComLibrary)
+		{
+			LogError(LOG_PREAMBLE, "Could not initialize COM library.");
+			return;
+		}
+
+		auto GetDllHandle = [&DllHandles = DllHandles](const TCHAR* DllName) -> void* {
+			void* DllHandle = FWindowsPlatformProcess::GetDllHandle(DllName);
+			if (DllHandle)
+			{
+				DllHandles.Add(DllHandle);
+				LogVerbose(LOG_PREAMBLE, FString::Format(TEXT("Retrieved DLL handle for {0}."), {DllName}));
+			}
+			else
+			{
+				LogError(LOG_PREAMBLE, FString::Format(TEXT("Could not get DLL handle for {0}."), {DllName}));
+			}
+			return DllHandle;
+		};
+
+		auto GetDllExport = []<typename FFuncType>(const TCHAR* DllName, void* DllHandle, const TCHAR* ExportName, FFuncType& ExportHandle)
+		{
+			void* FuncPtr = FWindowsPlatformProcess::GetDllExport(DllHandle, ExportName);
+			if (FuncPtr)
+			{
+				LogVerbose(LOG_PREAMBLE, FString::Format(TEXT("Retrieved DLL export '{0}' from {1}."), {ExportName, DllName}));
+			}
+			else
+			{
+				LogError(LOG_PREAMBLE, FString::Format(TEXT("Could not get DLL export '{0}' in {1}."), {ExportName, DllName}));
+			}
+			ExportHandle = static_cast<FFuncType>(FuncPtr);
+		};
+
+		const TCHAR* Kernel32DllNAme = TEXT("Kernel32.dll");
+		void* Kernel32DllHandle = GetDllHandle(Kernel32DllNAme);
+		GetDllExport(Kernel32DllNAme, Kernel32DllHandle, TEXT("GetTickCount"), GetTickCount);
+
+		const TCHAR* User32DllName = TEXT("User32.dll");
+		void* User32DllHandle = GetDllHandle(User32DllName);
+		GetDllExport(User32DllName, User32DllHandle, TEXT("GetClientRect"), GetClientRect);
+		GetDllExport(User32DllName, User32DllHandle, TEXT("GetDC"), GetDC);
+		GetDllExport(User32DllName, User32DllHandle, TEXT("ReleaseDC"), ReleaseDC);
+
+		const TCHAR* Gdi32DllName = TEXT("Gdi32.dll");
+		void* Gdi32DllHandle = GetDllHandle(Gdi32DllName);
+		GetDllExport(Gdi32DllName, Gdi32DllHandle, TEXT("GetDeviceCaps"), GetDeviceCaps);
+
+		const TCHAR* Ole32DllName = TEXT("Ole32.dll");
+		void* Ole32DllHandle = GetDllHandle(Ole32DllName);
+		GetDllExport(Ole32DllName, Ole32DllHandle, TEXT("CoCreateInstance"), CoCreateInstance);
+		GetDllExport(Ole32DllName, Ole32DllHandle, TEXT("CoCreateFreeThreadedMarshaler"), CoCreateFreeThreadedMarshaler);
+		GetDllExport(Ole32DllName, Ole32DllHandle, TEXT("CoTaskMemFree"), CoTaskMemFree);
+		GetDllExport(Ole32DllName, Ole32DllHandle, TEXT("StringFromGUID2"), StringFromGUID2);
+
+		const TCHAR* OleAut32DllName = TEXT("OleAut32.dll");
+		void* OleAut32DllHandle = GetDllHandle(OleAut32DllName);
+		GetDllExport(OleAut32DllName, OleAut32DllHandle, TEXT("SysFreeString"), SysFreeString);
+		GetDllExport(OleAut32DllName, OleAut32DllHandle, TEXT("VariantClear"), VariantClear);
+		GetDllExport(OleAut32DllName, OleAut32DllHandle, TEXT("VariantInit"), VariantInit);
+
+		// Try to get handles to RTSCom.dll and InkObj.dll for WindowsInk.
+		{
+			const TCHAR* Shell32DllName = TEXT("Shell32.dll");
+			void* Shell32DllHandle = GetDllHandle(Shell32DllName);
+
+			using FFuncSHGetKnownFolderPath = HRESULT(REFKNOWNFOLDERID, DWORD, HANDLE, PWSTR*);
+			FFuncSHGetKnownFolderPath* SHGetKnownFolderPath = nullptr;
+			GetDllExport(Shell32DllName, Shell32DllHandle, TEXT("SHGetKnownFolderPath"), SHGetKnownFolderPath);
+
+			if (SHGetKnownFolderPath && CoTaskMemFree)
+			{
+				PWSTR ProgramFilesCommonPath = nullptr;
+				if (Succeeded(SHGetKnownFolderPath(FOLDERID_ProgramFilesCommon, 0, nullptr, &ProgramFilesCommonPath), LOG_PREAMBLE))
+				{
+					const FString InkDLLDirectory = FString(ProgramFilesCommonPath) + TEXT("\\microsoft shared\\ink");
+					FPlatformProcess::PushDllDirectory(*InkDLLDirectory);
+					{
+						const TCHAR* RTSComDllName = TEXT("RTSCom.dll");
+						void* RTSComDllHandle = GetDllHandle(RTSComDllName);
+						bHasRTSComDllHandle = RTSComDllHandle != nullptr;
+
+						const TCHAR* InkObjDllName = TEXT("InkObj.dll");
+						void* InkObjDllHandle = GetDllHandle(InkObjDllName);
+						bHasInkObjDllHandle = InkObjDllHandle != nullptr;
+					}
+					FPlatformProcess::PopDllDirectory(*InkDLLDirectory);
+				}
+
+				CoTaskMemFree(ProgramFilesCommonPath);
+			}
+		}
+	}
+
+	FRealTimeStylusAPI::~FRealTimeStylusAPI()
+	{
+		for (void* DllHandle : DllHandles)
+		{
+			FWindowsPlatformProcess::FreeDllHandle(DllHandle);
+		}
+
+		if (bInitializedComLibrary)
+		{
+			FWindowsPlatformMisc::CoUninitialize();
+		}
+	}
+
+	const FRealTimeStylusAPI& FRealTimeStylusAPI::GetInstance()
+	{
+		static FRealTimeStylusAPI API;
+		return API;
+	}
+
+	bool FRealTimeStylusAPI::SatisfiesRequirements() const
+	{
+		bool bSuccess = true;
+
+		bSuccess &= bInitializedComLibrary;
+		bSuccess &= GetTickCount != nullptr;
+		bSuccess &= GetClientRect != nullptr;
+		bSuccess &= GetDC != nullptr;
+		bSuccess &= ReleaseDC != nullptr;
+		bSuccess &= GetDeviceCaps != nullptr;
+		bSuccess &= CoCreateInstance != nullptr;
+		bSuccess &= CoCreateFreeThreadedMarshaler != nullptr;
+		bSuccess &= CoTaskMemFree != nullptr;
+		bSuccess &= StringFromGUID2 != nullptr;
+		bSuccess &= SysFreeString != nullptr;
+		bSuccess &= VariantClear != nullptr;
+		bSuccess &= VariantInit != nullptr;
+		bSuccess &= bHasRTSComDllHandle;
+		bSuccess &= bHasInkObjDllHandle;
+
+		return bSuccess;
+	}
+}
+
+#undef LOG_PREAMBLE
+
+#endif
diff --git a/Engine/Plugins/Editor/StylusInput/Source/StylusInput/Private/Windows/RealTimeStylus/RealTimeStylusAPI.h b/Engine/Plugins/Editor/StylusInput/Source/StylusInput/Private/Windows/RealTimeStylus/RealTimeStylusAPI.h
new file mode 100644
--- /dev/null
+++ b/Engine/Plugins/Editor/StylusInput/Source/StylusInput/Private/Windows/RealTimeStylus/RealTimeStylusAPI.h
@@ -0,0 +1,76 @@
+﻿// Copyright Epic Games, Inc. All Rights Reserved.
+
+#pragma once
+
+#include <HAL/Platform.h>
+
+#if PLATFORM_WINDOWS
+
+#include <Containers/Array.h>
+
+#include <Windows/AllowWindowsPlatformTypes.h>
+	#include <windef.h>
+	#include <wtypes.h>
+	#include <Unknwn.h>
+#include <Windows/HideWindowsPlatformTypes.h>
+
+namespace UE::StylusInput::Private::RealTimeStylus
+{
+	class FRealTimeStylusAPI
+	{
+	public:
+		static const FRealTimeStylusAPI& GetInstance();
+
+		bool SatisfiesRequirements() const;
+
+		using FFuncGetTickCount = DWORD();
+		FFuncGetTickCount* GetTickCount = nullptr;
+
+		using FFuncGetClientRect = BOOL(HWND, LPRECT);
+		FFuncGetClientRect* GetClientRect = nullptr;
+
+		using FFuncGetDC = HDC(HWND);
+		FFuncGetDC* GetDC = nullptr;
+
+		using FFuncReleaseDC = int(HWND, HDC);
+		FFuncReleaseDC* ReleaseDC = nullptr;
+
+		using FFuncGetDeviceCaps = int(HDC, int);
+		FFuncGetDeviceCaps* GetDeviceCaps = nullptr;
+
+		using FFuncCoCreateInstance = HRESULT(REFCLSID, IUnknown*, DWORD, REFIID, LPVOID);
+		FFuncCoCreateInstance* CoCreateInstance = nullptr;
+
+		using FFuncCoCreateFreeThreadedMarshaler = HRESULT(LPUNKNOWN, LPUNKNOWN*);
+		FFuncCoCreateFreeThreadedMarshaler* CoCreateFreeThreadedMarshaler = nullptr;
+
+		using FFuncCoTaskMemFree = void(LPVOID);
+		FFuncCoTaskMemFree* CoTaskMemFree = nullptr;
+
+		using FFuncStringFromGUID2 = int(REFGUID, LPOLESTR, int);
+		FFuncStringFromGUID2* StringFromGUID2 = nullptr;
+
+		using FFuncSysFreeString = void(BSTR);
+		FFuncSysFreeString* SysFreeString = nullptr;
+
+		using FFuncVariantClear = HRESULT(VARIANTARG*);
+		FFuncVariantClear* VariantClear = nullptr;
+
+		using FFuncVariantInit = void(VARIANTARG*);
+		FFuncVariantInit* VariantInit = nullptr;
+
+	private:
+
+		FRealTimeStylusAPI();
+		~FRealTimeStylusAPI();
+
+		TArray<void*> DllHandles;
+
+		bool bInitializedComLibrary = false;
+		bool bHasRTSComDllHandle = false;
+		bool bHasInkObjDllHandle = false;
+
+	};
+}
+
+#endif
diff --git a/Engine/Plugins/Editor/StylusInput/Source/StylusInput/Private/Windows/RealTimeStylus/RealTimeStylusInstance.cpp b/Engine/Plugins/Editor/StylusInput/Source/StylusInput/Private/Windows/RealTimeStylus/RealTimeStylusInstance.cpp
new file mode 100644
--- /dev/null
+++ b/Engine/Plugins/Editor/StylusInput/Source/StylusInput/Private/Windows/RealTimeStylus/RealTimeStylusInstance.cpp
@@ -0,0 +1,470 @@
+﻿// Copyright Epic Games, Inc. All Rights Reserved.
+
+#include "RealTimeStylusInstance.h"
+
+#if PLATFORM_WINDOWS
+
+#include <StylusInput.h>
+#include <StylusInputUtils.h>
+#include <Algo/Transform.h>
+
+#include "RealTimeStylusAPI.h"
+#include "RealTimeStylusPluginAsync.h"
+#include "RealTimeStylusPluginSync.h"
+#include "RealTimeStylusUtils.h"
+
+// Fixed indices at which the plugins are installed; we only ever have one each.
+#define ASYNC_PLUGIN_INDEX 0
+#define SYNC_PLUGIN_INDEX 0
+
+#define LOCTEXT_NAMESPACE "RealTimeStylusInstance"
+#define LOG_PREAMBLE "RealTimeStylusInstance"
+
+namespace UE::StylusInput::Private::RealTimeStylus
+{
+	FRealTimeStylusInstance::FRealTimeStylusInstance(const uint32 ID, const HWND OSWindowHandle)
+		: ID(ID)
+		, RealTimeStylusAPI(FRealTimeStylusAPI::GetInstance())
+	{
+		Init(OSWindowHandle);
+	}
+
+	FRealTimeStylusInstance::~FRealTimeStylusInstance()
+	{
+		RealTimeStylus->RemoveAllStylusAsyncPlugins();
+		RealTimeStylus->RemoveAllStylusSyncPlugins();
+		RealTimeStylus->Release();
+	}
+
+	bool FRealTimeStylusInstance::AddEventHandler(IStylusInputEventHandler* EventHandler, const EEventHandlerThread Thread)
+	{
+		if (!EventHandler)
+		{
+			LogWarning(LOG_PREAMBLE, "Tried to add nullptr as event handler.");
+			return false;
+		}
+
+		if (Thread == EEventHandlerThread::OnGameThread)
+		{
+			if (!AsyncPlugin)
+			{
+				EnablePlugin(EEventHandlerThread::OnGameThread, EventHandler);
+
+				if (!AsyncPlugin)
+				{
+					LogError(LOG_PREAMBLE, FString::Format(
+								 TEXT("Event handler '{0}' was not added since the game thread stylus input plugin could not be installed."),
+								 {EventHandler->GetName()}));
+					return false;
+				}
+
+				return true;
+			}
+
+			return AsyncPlugin->AddEventHandler(EventHandler);
+		}
+
+		if (Thread == EEventHandlerThread::Asynchronous)
+		{
+			if (!SyncPlugin)
+			{
+				EnablePlugin(EEventHandlerThread::Asynchronous, EventHandler);
+
+				if (!SyncPlugin)
+				{
+					LogError(LOG_PREAMBLE, FString::Format(
+								 TEXT("Event handler '{0}' was not added since the asynchronous stylus input plugin could not be installed."),
+								 {EventHandler->GetName()}));
+					return false;
+				}
+
+				return true;
+			}
+
+			return SyncPlugin->AddEventHandler(EventHandler);
+		}
+
+		return false;
+	}
+
+	bool FRealTimeStylusInstance::RemoveEventHandler(IStylusInputEventHandler* EventHandler)
+	{
+		if (!EventHandler)
+		{
+			LogWarning(LOG_PREAMBLE, "Tried to remove nullptr event handler.");
+			return false;
+		}
+
+		bool bWasRemoved = false;
+
+		if (AsyncPlugin && AsyncPlugin->RemoveEventHandler(EventHandler))
+		{
+			bWasRemoved = true;
+
+			if (AsyncPlugin->NumEventHandlers() == 0)
+			{
+				DisablePlugin(EEventHandlerThread::OnGameThread);
+			}
+		}
+
+		if (SyncPlugin && SyncPlugin->RemoveEventHandler(EventHandler))
+		{
+			bWasRemoved = true;
+
+			if (SyncPlugin->NumEventHandlers() == 0)
+			{
+				DisablePlugin(EEventHandlerThread::Asynchronous);
+			}
+		}
+
+		if (!bWasRemoved)
+		{
+			LogError(LOG_PREAMBLE, FString::Format(TEXT("Event handler {0} does not exist."), {EventHandler->GetName()}));
+		}
+
+		return bWasRemoved;
+	}
+
+	const TSharedPtr<IStylusInputTabletContext> FRealTimeStylusInstance::GetTabletContext(const uint32 TabletContextID)
+	{
+		return TabletContexts.Get(TabletContextID);
+	}
+
+	const TSharedPtr<IStylusInputStylusInfo> FRealTimeStylusInstance::GetStylusInfo(uint32 StylusID)
+	{
+		TSharedPtr<FStylusInfo> StylusInfo = StylusInfos.Get(StylusID);
+		if (StylusInfo.IsValid())
+		{
+			return StylusInfo;
+		}
+
+		StylusInfo = StylusInfos.Add(StylusID);
+
+		IInkCursor* InkCursor = nullptr;
+		if (Succeeded(RealTimeStylus->GetStylusForId(StylusID, &InkCursor), LOG_PREAMBLE) && InkCursor != nullptr)
+		{
+			BSTR String = nullptr;
+			if (Succeeded(InkCursor->get_Name(&String), LOG_PREAMBLE))
+			{
+				StylusInfo->Name = String;
+				RealTimeStylusAPI.SysFreeString(String);
+			}
+			else
+			{
+				LogError(LOG_PREAMBLE, FString::Format(TEXT("Could not get name for stylus with ID {0}."), {StylusID}));
+			}
+
+			IInkCursorButtons* Buttons = nullptr;
+			long NumButtons = 0;
+			if (Succeeded(InkCursor->get_Buttons(&Buttons), LOG_PREAMBLE) &&
+				Buttons && 
+				Succeeded(Buttons->get_Count(&NumButtons), LOG_PREAMBLE))
+			{
+				if (NumButtons > 0)
+				{
+					TArray<FStylusButton> StylusInfoButtons;
+					StylusInfoButtons.Reserve(NumButtons);
+
+					for (long ButtonIndex = 0; ButtonIndex < NumButtons; ++ButtonIndex)
+					{
+						VARIANT ButtonIdentifier;
+						RealTimeStylusAPI.VariantInit(&ButtonIdentifier);
+
+						ButtonIdentifier.vt = VT_INT;
+						ButtonIdentifier.lVal = ButtonIndex;
+
+						IInkCursorButton* Button = nullptr;
+						if (Succeeded(Buttons->Item(ButtonIdentifier, &Button), LOG_PREAMBLE))
+						{
+							FStylusButton NewButton;
+
+							if (Succeeded(Button->get_Id(&String), LOG_PREAMBLE))
+							{
+								NewButton.ID = String;
+								RealTimeStylusAPI.SysFreeString(String);
+							}
+							else
+							{
+								LogError(LOG_PREAMBLE, FString::Format(
+											 TEXT("Could not get ID for button {0} for stylus with ID {1}."), {static_cast<int32>(ButtonIndex), StylusID}));
+							}
+
+							if (Succeeded(Button->get_Name(&String), LOG_PREAMBLE))
+							{
+								NewButton.Name = String;
+								RealTimeStylusAPI.SysFreeString(String);
+							}
+							else
+							{
+								LogError(LOG_PREAMBLE, FString::Format(
+											 TEXT("Could not get name for button {0} for stylus with ID {1}."), {static_cast<int32>(ButtonIndex), StylusID}));
+							}
+
+							if (!NewButton.ID.IsEmpty() && !NewButton.Name.IsEmpty())
+							{
+								StylusInfoButtons.Emplace(NewButton);
+							}
+						}
+						else
+						{
+							LogError(LOG_PREAMBLE, FString::Format(
+										 TEXT("Could not get button {0} for stylus with ID {1}."), {static_cast<int32>(ButtonIndex), StylusID}));
+						}
+
+						RealTimeStylusAPI.VariantClear(&ButtonIdentifier);
+					}
+
+					StylusInfo->Buttons = MoveTemp(StylusInfoButtons);
+				}
+			}
+			else
+			{
+				LogError(LOG_PREAMBLE, FString::Format(TEXT("Could not get buttons for stylus with ID {0}."), {StylusID}));
+			}
+		}
+		else
+		{
+			LogWarning(LOG_PREAMBLE, FString::Format(TEXT("Could not get stylus info for ID {0}."), {StylusID}));
+		}
+
+		return StylusInfo;
+	}
+
+	float FRealTimeStylusInstance::GetPacketsPerSecond(const EEventHandlerThread EventHandlerThread) const
+	{
+		if (EventHandlerThread == EEventHandlerThread::Asynchronous)
+		{
+			return SyncPlugin ? SyncPlugin->GetPacketsPerSecond() : 0.0f;
+		}
+
+		if (EventHandlerThread == EEventHandlerThread::OnGameThread)
+		{
+			return AsyncPlugin ? AsyncPlugin->GetPacketsPerSecond() : 0.0f;
+		}
+
+		return -1.0f;
+	}
+
+	FText FRealTimeStylusInstance::GetName()
+	{
+		return FText::Format(LOCTEXT("RealTimeStylus", "RealTimeStylus #{0}"), ID);
+	}
+
+	void FRealTimeStylusInstance::SetupWindowContext(const HWND HWindow)
+	{
+		RECT WindowRectangle;
+		if (!ensure(RealTimeStylusAPI.GetClientRect(HWindow, &WindowRectangle)))
+		{
+			LogError(LOG_PREAMBLE, "Could not retrieve window rectangle; failed to setup device context!");
+			return;
+		}
+
+		if (WindowRectangle.right <= WindowRectangle.left || WindowRectangle.bottom <= WindowRectangle.top)
+		{
+			LogError(LOG_PREAMBLE, "Window appears to be minimized; failed to setup device context!");
+			return;
+		}
+
+		const HDC HWindowDeviceContext = RealTimeStylusAPI.GetDC(HWindow);
+		if (!ensure(HWindowDeviceContext))
+		{
+			LogError(LOG_PREAMBLE, "Could not retrieve window device context; failed to setup device context!");
+			return;
+		}
+
+		const FVector2d DisplaySizeInMillimeters = {
+			static_cast<double>(RealTimeStylusAPI.GetDeviceCaps(HWindowDeviceContext, HORZSIZE)),
+			static_cast<double>(RealTimeStylusAPI.GetDeviceCaps(HWindowDeviceContext, VERTSIZE))
+		};
+
+		if (!ensure(DisplaySizeInMillimeters.X > 0.0 && DisplaySizeInMillimeters.Y > 0.0))
+		{
+			LogError(LOG_PREAMBLE, "Display size in millimeters is invalid; failed to setup device context!");
+			return;
+		}
+
+		const FVector2d DisplaySizeInPixels = {
+			static_cast<double>(RealTimeStylusAPI.GetDeviceCaps(HWindowDeviceContext, HORZRES)),
+			static_cast<double>(RealTimeStylusAPI.GetDeviceCaps(HWindowDeviceContext, VERTRES))
+		};
+
+		if (!ensure(DisplaySizeInPixels.X > 0.0 && DisplaySizeInPixels.Y > 0.0))
+		{
+			LogError(LOG_PREAMBLE, "Display size in pixels is invalid; failed to setup device context!");
+			return;
+		}
+
+		const FVector2d DisplayPixelsPerLogicalInch = {
+			static_cast<double>(RealTimeStylusAPI.GetDeviceCaps(HWindowDeviceContext, LOGPIXELSX)),
+			static_cast<double>(RealTimeStylusAPI.GetDeviceCaps(HWindowDeviceContext, LOGPIXELSY))
+		};
+
+		if (!ensure(DisplayPixelsPerLogicalInch.X > 0.0 && DisplayPixelsPerLogicalInch.Y > 0.0))
+		{
+			LogError(LOG_PREAMBLE, "Display pixels per logical inch is invalid; failed to setup device context!");
+			return;
+		}
+
+		const FVector2d WindowSizeInPixels = {
+			static_cast<double>(WindowRectangle.right - WindowRectangle.left),
+			static_cast<double>(WindowRectangle.bottom - WindowRectangle.top)
+		};
+
+		const FVector2d DisplayMillimetersPerPixel = DisplaySizeInMillimeters / DisplaySizeInPixels;
+		const FVector2d WindowSizeInMillimeters = WindowSizeInPixels * DisplayMillimetersPerPixel;
+		const FVector2d MillimetersPerInch{25.4};
+		const FVector2d DisplayInchesPerPixel = MillimetersPerInch / DisplayMillimetersPerPixel;
+
+		WindowContext.XYScale = DisplayInchesPerPixel / DisplayPixelsPerLogicalInch;
+		WindowContext.XYMaximum = WindowSizeInMillimeters * FVector2d{100.0} + FVector2d{0.5};
+		WindowContext.WindowSize = WindowSizeInPixels;
+
+		RealTimeStylusAPI.ReleaseDC(HWindow, HWindowDeviceContext);
+
+		LogVerbose(LOG_PREAMBLE, "Successfully setup window context.");
+	}
+
+	const FWindowContext& FRealTimeStylusInstance::GetWindowContext() const
+	{
+		return WindowContext;
+	}
+
+	void FRealTimeStylusInstance::UpdateTabletContexts(const FTabletContextContainer& InTabletContexts)
+	{
+		TabletContexts.Update(InTabletContexts);
+	}
+
+	void FRealTimeStylusInstance::Init(const HWND OSWindowHandle)
+	{
+		check(!RealTimeStylus.IsValid());
+
+		SetupWindowContext(OSWindowHandle);
+
+		void* OutInstance = nullptr;
+		if (Failed(RealTimeStylusAPI.CoCreateInstance(CLSID_RealTimeStylus, nullptr, CLSCTX_ALL, IID_IRealTimeStylus, &OutInstance), LOG_PREAMBLE))
+		{
+			LogError(LOG_PREAMBLE, "Could not create RealTimeStylus COM object instance!");
+		}
+		RealTimeStylus = static_cast<IRealTimeStylus*>(OutInstance);
+
+		if (RealTimeStylus.IsValid())
+		{
+			TArray<GUID> DesiredPackets;
+			DesiredPackets.Reserve(std::size(PacketPropertyConstants));
+			Algo::Transform(PacketPropertyConstants, DesiredPackets, [](const FPacketPropertyConstant& Property) { return Property.Guid; });
+
+			if (Failed(RealTimeStylus->SetDesiredPacketDescription(DesiredPackets.Num(), DesiredPackets.GetData()), LOG_PREAMBLE))
+			{
+				LogError(LOG_PREAMBLE, "Could not set desired packet description!");
+			}
+
+			if (Failed(RealTimeStylus->put_HWND(reinterpret_cast<HANDLE_PTR>(OSWindowHandle)), LOG_PREAMBLE))
+			{
+				LogError(LOG_PREAMBLE, "Could not set window handle!");
+			}
+
+			if (Failed(RealTimeStylus->put_Enabled(::Windows::TRUE), LOG_PREAMBLE))
+			{
+				LogError(LOG_PREAMBLE, "Could not enable real time stylus input!");
+			}
+
+			LogVerbose(LOG_PREAMBLE, "Successfully initialized RealTimeStylus COM object.");
+		}
+	}
+
+	void FRealTimeStylusInstance::EnablePlugin(const EEventHandlerThread EventHandlerThread, IStylusInputEventHandler* EventHandler)
+	{
+		check(RealTimeStylus.IsValid());
+
+		if (EventHandlerThread == EEventHandlerThread::OnGameThread)
+		{
+			AsyncPlugin = MakeUnique<FRealTimeStylusPluginAsync>(
+				this,
+				FGetWindowContextCallback::CreateRaw(this, &FRealTimeStylusInstance::GetWindowContext),
+				FUpdateTabletContextsCallback::CreateRaw(this, &FRealTimeStylusInstance::UpdateTabletContexts),
+				EventHandler
+			);
+
+			if (Succeeded(RealTimeStylus->AddStylusAsyncPlugin(ASYNC_PLUGIN_INDEX, AsyncPlugin.Get()), LOG_PREAMBLE))
+			{
+				LogVerbose(LOG_PREAMBLE, "Added game thread stylus input plugin!");
+			}
+			else
+			{
+				LogError(LOG_PREAMBLE, "Could not add game thread stylus input plugin!");
+				AsyncPlugin.Reset();
+			}
+		}
+
+		if (EventHandlerThread == EEventHandlerThread::Asynchronous)
+		{
+			SyncPlugin = MakeUnique<FRealTimeStylusPluginSync>(
+				this,
+				FGetWindowContextCallback::CreateRaw(this, &FRealTimeStylusInstance::GetWindowContext),
+				FUpdateTabletContextsCallback::CreateRaw(this, &FRealTimeStylusInstance::UpdateTabletContexts),
+				EventHandler
+			);
+
+			if (Succeeded(SyncPlugin->CreateFreeThreadMarshaler(), LOG_PREAMBLE))
+			{
+				if (Succeeded(RealTimeStylus->AddStylusSyncPlugin(SYNC_PLUGIN_INDEX, SyncPlugin.Get()), LOG_PREAMBLE))
+				{
+					LogVerbose(LOG_PREAMBLE, "Added asynchronous stylus input plugin!");
+				}
+				else
+				{
+					LogError(LOG_PREAMBLE, "Could not add asynchronous stylus input plugin!");
+					SyncPlugin.Reset();
+				}
+			}
+			else
+			{
+				LogError(LOG_PREAMBLE, "Could not create free thread marshaler for asynchronous stylus input plugin!");
+				SyncPlugin.Reset();
+			}
+		}
+	}
+
+	void FRealTimeStylusInstance::DisablePlugin(const EEventHandlerThread EventHandlerThread)
+	{
+		check(RealTimeStylus.IsValid());
+
+#pragma warning(push)
+#pragma warning(disable : 6387)
+
+		if (EventHandlerThread == EEventHandlerThread::OnGameThread)
+		{
+			if (Failed(RealTimeStylus->RemoveStylusAsyncPlugin(ASYNC_PLUGIN_INDEX, nullptr), LOG_PREAMBLE))
+			{
+				LogError(LOG_PREAMBLE, "Could not remove game thread stylus input plugin!");
+			}
+			else
+			{
+				LogVerbose(LOG_PREAMBLE, "Removed game thread stylus input plugin!");
+			}
+
+			AsyncPlugin.Reset();
+		}
+
+		if (EventHandlerThread == EEventHandlerThread::Asynchronous)
+		{
+			if (Failed(RealTimeStylus->RemoveStylusSyncPlugin(SYNC_PLUGIN_INDEX, nullptr), LOG_PREAMBLE))
+			{
+				LogError(LOG_PREAMBLE, "Could not remove asynchronous stylus input plugin!");
+			}
+			else
+			{
+				LogVerbose(LOG_PREAMBLE, "Removed asynchronous stylus input plugin!");
+			}
+
+			SyncPlugin.Reset();
+		}
+
+#pragma warning(pop)
+	}
+}
+
+#undef LOG_PREAMBLE
+#undef LOCTEXT_NAMESPACE
+
+#endif
diff --git a/Engine/Plugins/Editor/StylusInput/Source/StylusInput/Private/Windows/RealTimeStylus/RealTimeStylusInstance.h b/Engine/Plugins/Editor/StylusInput/Source/StylusInput/Private/Windows/RealTimeStylus/RealTimeStylusInstance.h
new file mode 100644
--- /dev/null
+++ b/Engine/Plugins/Editor/StylusInput/Source/StylusInput/Private/Windows/RealTimeStylus/RealTimeStylusInstance.h
@@ -0,0 +1,71 @@
+﻿// Copyright Epic Games, Inc. All Rights Reserved.
+
+#pragma once
+
+#include <HAL/Platform.h>
+
+#if PLATFORM_WINDOWS
+
+#include <StylusInput.h>
+#include <Containers/Array.h>
+#include <Templates/UniquePtr.h>
+
+#include <Windows/AllowWindowsPlatformTypes.h>
+	#include <Microsoft/COMPointer.h>
+	#include <RTSCom.h>
+#include <Windows/HideWindowsPlatformTypes.h>
+
+#include "RealTimeStylusTabletContext.h"
+
+namespace UE::StylusInput
+{
+	class IStylusInputEventHandler;
+}
+
+namespace UE::StylusInput::Private::RealTimeStylus
+{
+	class FRealTimeStylusPluginSync;
+	class FRealTimeStylusPluginAsync;
+	class FRealTimeStylusAPI;
+
+	class FRealTimeStylusInstance : public IStylusInputInstance
+	{
+	public:
+		explicit FRealTimeStylusInstance(uint32 ID, HWND OSWindowHandle);
+		virtual ~FRealTimeStylusInstance() override;
+
+		virtual bool AddEventHandler(IStylusInputEventHandler* EventHandler, EEventHandlerThread Thread) override;
+		virtual bool RemoveEventHandler(IStylusInputEventHandler* EventHandler) override;
+
+		virtual const TSharedPtr<IStylusInputTabletContext> GetTabletContext(uint32 TabletContextID) override;
+		virtual const TSharedPtr<IStylusInputStylusInfo> GetStylusInfo(uint32 StylusID) override;
+
+		virtual float GetPacketsPerSecond(EEventHandlerThread EventHandlerThread) const override;
+
+		virtual FText GetName() override;
+
+	private:
+		void Init(HWND OSWindowHandle);
+
+		void EnablePlugin(EEventHandlerThread EventHandlerThread, IStylusInputEventHandler* EventHandler);
+		void DisablePlugin(EEventHandlerThread EventHandlerThread);
+
+		void SetupWindowContext(HWND HWindow);
+		const FWindowContext& GetWindowContext() const;
+
+		void UpdateTabletContexts(const FTabletContextContainer& InTabletContexts);
+
+		const uint32 ID;
+		const FRealTimeStylusAPI& RealTimeStylusAPI;
+		FWindowContext WindowContext;
+
+		TComPtr<IRealTimeStylus> RealTimeStylus;
+		TUniquePtr<FRealTimeStylusPluginAsync> AsyncPlugin;
+		TUniquePtr<FRealTimeStylusPluginSync> SyncPlugin;
+
+		FTabletContextThreadSafeContainer TabletContexts;
+		FStylusInfoThreadSafeContainer StylusInfos;
+	};
+}
+
+#endif
diff --git a/Engine/Plugins/Editor/StylusInput/Source/StylusInput/Private/Windows/RealTimeStylus/RealTimeStylusInterface.cpp b/Engine/Plugins/Editor/StylusInput/Source/StylusInput/Private/Windows/RealTimeStylus/RealTimeStylusInterface.cpp
new file mode 100644
--- /dev/null
+++ b/Engine/Plugins/Editor/StylusInput/Source/StylusInput/Private/Windows/RealTimeStylus/RealTimeStylusInterface.cpp
@@ -0,0 +1,95 @@
+﻿// Copyright Epic Games, Inc. All Rights Reserved.
+
+#include "RealTimeStylusInterface.h"
+
+#if PLATFORM_WINDOWS
+
+#include <StylusInput.h>
+#include <StylusInputUtils.h>
+
+#include <GenericPlatform/GenericWindow.h>
+#include <Templates/SharedPointer.h>
+#include <Widgets/SWindow.h>
+
+#include <Windows/AllowWindowsPlatformTypes.h>
+	#include <RTSCOM_i.c>
+#include <Windows/HideWindowsPlatformTypes.h>
+
+#include "RealTimeStylusInstance.h"
+#include "RealTimeStylusAPI.h"
+
+namespace UE::StylusInput::Private::RealTimeStylus
+{
+	IStylusInputInstance* FRealTimeStylusInterface::CreateInstance(SWindow& Window)
+	{
+		if (FRefCountedInstance* ExistingRefCountedInstance = Instances.Find(&Window))
+		{
+			++ExistingRefCountedInstance->RefCount;
+			return ExistingRefCountedInstance->Instance.Get();
+		}
+
+		HWND OSWindowHandle = [&Window]
+		{
+			const TSharedPtr<const FGenericWindow> NativeWindow = Window.GetNativeWindow();
+			return NativeWindow.IsValid() ? static_cast<HWND>(NativeWindow->GetOSWindowHandle()) : nullptr;
+		}();
+
+		if (!OSWindowHandle)
+		{
+			LogError("RealTimeStylusInterface", "Could not get native window handle.");
+			return nullptr;
+		}
+
+		FRealTimeStylusInstance* NewInstance = Instances.Emplace(
+			&Window, {MakeUnique<FRealTimeStylusInstance>(NextInstanceID++, OSWindowHandle), 1}).Instance.Get();
+		if (!ensureMsgf(NewInstance, TEXT("RealTimeStylusInterface: Failed to create stylus input instance.")))
+		{
+			Instances.Remove(&Window);
+			return nullptr;
+		}
+
+		return NewInstance;
+	}
+
+	bool FRealTimeStylusInterface::ReleaseInstance(IStylusInputInstance* Instance)
+	{
+		check(Instance);
+
+		FRealTimeStylusInstance *const RealTimeStylusInstance = static_cast<FRealTimeStylusInstance*>(Instance);
+
+		// Find existing instance
+		for (TTuple<SWindow*, FRefCountedInstance>& Entry : Instances)
+		{
+			FRefCountedInstance& RefCountedInstance = Entry.Get<1>();
+
+			if (RefCountedInstance.Instance.Get() == RealTimeStylusInstance)
+			{
+				// Decrease reference count
+				check(RefCountedInstance.RefCount > 0);
+				if (--RefCountedInstance.RefCount == 0)
+				{
+					// Delete if there are no references left
+					Instances.Remove(Entry.Key);
+				}
+
+				return true;
+			}
+		}
+
+		ensureMsgf(false, TEXT("RealTimeStylusInterface: Failed to find provided instance."));
+		return false;
+	}
+
+	TUniquePtr<IStylusInputInterface> FRealTimeStylusInterface::Create()
+	{
+		if (FRealTimeStylusAPI::GetInstance().SatisfiesRequirements())
+		{
+			TUniquePtr<FRealTimeStylusInterface> RealTimeStylusInterface = MakeUnique<FRealTimeStylusInterface>();
+			return TUniquePtr<IStylusInputInterface>(MoveTemp(RealTimeStylusInterface));
+		}
+
+		return nullptr;
+	}
+}
+
+#endif
diff --git a/Engine/Plugins/Editor/StylusInput/Source/StylusInput/Private/Windows/RealTimeStylus/RealTimeStylusInterface.h b/Engine/Plugins/Editor/StylusInput/Source/StylusInput/Private/Windows/RealTimeStylus/RealTimeStylusInterface.h
new file mode 100644
--- /dev/null
+++ b/Engine/Plugins/Editor/StylusInput/Source/StylusInput/Private/Windows/RealTimeStylus/RealTimeStylusInterface.h
@@ -0,0 +1,45 @@
+﻿// Copyright Epic Games, Inc. All Rights Reserved.
+
+#pragma once
+
+#include <HAL/Platform.h>
+
+#if PLATFORM_WINDOWS
+
+#include <StylusInputInterface.h>
+#include <Containers/Array.h>
+#include <Containers/Map.h>
+#include <Templates/UniquePtr.h>
+
+namespace UE::StylusInput::Private
+{
+	class FStylusInputImpl;
+}
+
+namespace UE::StylusInput::Private::RealTimeStylus
+{
+	class FRealTimeStylusInstance;
+
+	class FRealTimeStylusInterface : public IStylusInputInterface
+	{
+	public:
+		virtual IStylusInputInstance* CreateInstance(SWindow& Window) override;
+		virtual bool ReleaseInstance(IStylusInputInstance* Instance) override;
+
+	private:
+		friend FStylusInputImpl;
+
+		static TUniquePtr<IStylusInputInterface> Create();
+
+		struct FRefCountedInstance
+		{
+			TUniquePtr<FRealTimeStylusInstance> Instance;
+			int32 RefCount;
+		};
+
+		TMap<SWindow*, FRefCountedInstance> Instances;
+		uint32 NextInstanceID = 0;
+	};
+}
+
+#endif
diff --git a/Engine/Plugins/Editor/StylusInput/Source/StylusInput/Private/Windows/RealTimeStylus/RealTimeStylusPluginAsync.cpp b/Engine/Plugins/Editor/StylusInput/Source/StylusInput/Private/Windows/RealTimeStylus/RealTimeStylusPluginAsync.cpp
new file mode 100644
--- /dev/null
+++ b/Engine/Plugins/Editor/StylusInput/Source/StylusInput/Private/Windows/RealTimeStylus/RealTimeStylusPluginAsync.cpp
@@ -0,0 +1,121 @@
+﻿// Copyright Epic Games, Inc. All Rights Reserved.
+
+#include "RealTimeStylusPluginAsync.h"
+
+#if PLATFORM_WINDOWS
+
+namespace UE::StylusInput::Private::RealTimeStylus
+{
+	FRealTimeStylusPluginAsync::FRealTimeStylusPluginAsync(IStylusInputInstance* Instance,
+	                                                       FGetWindowContextCallback&& GetWindowContext,
+	                                                       FUpdateTabletContextsCallback&& UpdateTabletContextsCallback,
+	                                                       IStylusInputEventHandler* EventHandler)
+		: FRealTimeStylusPluginBase(Instance, MoveTemp(GetWindowContext), MoveTemp(UpdateTabletContextsCallback))
+	{
+		if (EventHandler)
+		{
+			// Immediately install an event handler during construction to capture events coming through during plugin initialization.
+			AddEventHandler(EventHandler);
+		}
+	}
+
+	FRealTimeStylusPluginAsync::~FRealTimeStylusPluginAsync()
+	{
+	}
+
+	HRESULT FRealTimeStylusPluginAsync::RealTimeStylusEnabled(IRealTimeStylus* RealTimeStylus, const ULONG TabletContextIDsCount,
+	                                                              const TABLET_CONTEXT_ID* TabletContextIDs)
+	{
+		return ProcessRealTimeStylusEnabled(RealTimeStylus, TabletContextIDsCount, TabletContextIDs);
+	}
+
+	HRESULT FRealTimeStylusPluginAsync::RealTimeStylusDisabled(IRealTimeStylus* RealTimeStylus, const ULONG TabletContextIDsCount,
+	                                                               const TABLET_CONTEXT_ID* TabletContextIDs)
+	{
+		return ProcessRealTimeStylusDisabled(RealTimeStylus, TabletContextIDsCount, TabletContextIDs);
+	}
+
+	HRESULT FRealTimeStylusPluginAsync::StylusDown(IRealTimeStylus* RealTimeStylus, const StylusInfo* StylusInfo, const ULONG PropertyCount,
+	                                                   LONG* PacketBuffer, LONG**)
+	{
+		return ProcessPackets(StylusInfo, 1, PropertyCount, EPacketType::StylusDown, reinterpret_cast<int32*>(PacketBuffer));
+	}
+
+	HRESULT FRealTimeStylusPluginAsync::StylusUp(IRealTimeStylus* RealTimeStylus, const StylusInfo* StylusInfo, const ULONG PropertyCount,
+	                                                 LONG* PacketBuffer, LONG**)
+	{
+		return ProcessPackets(StylusInfo, 1, PropertyCount, EPacketType::StylusUp, reinterpret_cast<int32*>(PacketBuffer));
+	}
+
+	HRESULT FRealTimeStylusPluginAsync::TabletAdded(IRealTimeStylus* RealTimeStylus, IInkTablet* Tablet)
+	{
+		return ProcessTabletAdded(Tablet);
+	}
+
+	HRESULT FRealTimeStylusPluginAsync::TabletRemoved(IRealTimeStylus* RealTimeStylus, const LONG TabletIndex)
+	{
+		return ProcessTabletRemoved(TabletIndex);
+	}
+
+	HRESULT FRealTimeStylusPluginAsync::InAirPackets(IRealTimeStylus* RealTimeStylus, const StylusInfo* StylusInfo, const ULONG PacketCount,
+	                                                const ULONG PacketBufferLength, LONG* PacketBuffer, ULONG*, LONG**)
+	{
+		return ProcessPackets(StylusInfo, PacketCount, PacketBufferLength, EPacketType::AboveDigitizer, reinterpret_cast<int32*>(PacketBuffer));
+	}
+
+	HRESULT FRealTimeStylusPluginAsync::Packets(IRealTimeStylus* RealTimeStylus, const StylusInfo* StylusInfo, const ULONG PacketCount,
+													const ULONG PacketBufferLength, LONG* PacketBuffer, ULONG*, LONG**)
+	{
+		return ProcessPackets(StylusInfo, PacketCount, PacketBufferLength, EPacketType::OnDigitizer, reinterpret_cast<int32*>(PacketBuffer));
+	}
+
+	HRESULT FRealTimeStylusPluginAsync::StylusInRange(IRealTimeStylus* RealTimeStylus, TABLET_CONTEXT_ID TabletContextID, STYLUS_ID StylusID)
+	{
+		return E_NOTIMPL;
+	}
+
+	HRESULT FRealTimeStylusPluginAsync::StylusOutOfRange(IRealTimeStylus* RealTimeStylus, TABLET_CONTEXT_ID TabletContextID, STYLUS_ID StylusID)
+	{
+		return E_NOTIMPL;
+	}
+
+	HRESULT FRealTimeStylusPluginAsync::StylusButtonDown(IRealTimeStylus* RealTimeStylus, STYLUS_ID StylusID, const GUID* GuidStylusButton,
+	                                                         POINT* StylusPos)
+	{
+		return E_NOTIMPL;
+	}
+
+	HRESULT FRealTimeStylusPluginAsync::StylusButtonUp(IRealTimeStylus* RealTimeStylus, STYLUS_ID StylusID, const GUID* GuidStylusButton, POINT* StylusPos)
+	{
+		return E_NOTIMPL;
+	}
+
+	HRESULT FRealTimeStylusPluginAsync::CustomStylusDataAdded(IRealTimeStylus* RealTimeStylus, const GUID* GuidId, ULONG DataCount, const BYTE* Data)
+	{
+		return E_NOTIMPL;
+	}
+
+	HRESULT FRealTimeStylusPluginAsync::SystemEvent(IRealTimeStylus* RealTimeStylus, TABLET_CONTEXT_ID TabletContextID, STYLUS_ID StylusID,
+	                                                    SYSTEM_EVENT Event, SYSTEM_EVENT_DATA EventData)
+	{
+		return E_NOTIMPL;
+	}
+
+	HRESULT FRealTimeStylusPluginAsync::Error(IRealTimeStylus*, IStylusPlugin*, const RealTimeStylusDataInterest DataInterest, const HRESULT ErrorCode,
+	                                              LONG_PTR*)
+	{
+		return ProcessError(DataInterest, ErrorCode);
+	}
+
+	HRESULT FRealTimeStylusPluginAsync::UpdateMapping(IRealTimeStylus* RealTimeStylus)
+	{
+		return E_NOTIMPL;
+	}
+
+	HRESULT FRealTimeStylusPluginAsync::DataInterest(RealTimeStylusDataInterest* DataInterest)
+	{
+		return ProcessDataInterest(DataInterest);
+	}
+}
+
+#endif
diff --git a/Engine/Plugins/Editor/StylusInput/Source/StylusInput/Private/Windows/RealTimeStylus/RealTimeStylusPluginAsync.h b/Engine/Plugins/Editor/StylusInput/Source/StylusInput/Private/Windows/RealTimeStylus/RealTimeStylusPluginAsync.h
new file mode 100644
--- /dev/null
+++ b/Engine/Plugins/Editor/StylusInput/Source/StylusInput/Private/Windows/RealTimeStylus/RealTimeStylusPluginAsync.h
@@ -0,0 +1,75 @@
+﻿// Copyright Epic Games, Inc. All Rights Reserved.
+
+#pragma once
+
+#include <HAL/Platform.h>
+
+#if PLATFORM_WINDOWS
+
+#include "RealTimeStylusPluginBase.h"
+
+namespace UE::StylusInput::Private::RealTimeStylus
+{
+	class FRealTimeStylusPluginAsync final : public IStylusAsyncPlugin, public FRealTimeStylusPluginBase
+	{
+	public:
+		FRealTimeStylusPluginAsync(IStylusInputInstance* Instance, FGetWindowContextCallback&& GetWindowContext,
+		                               FUpdateTabletContextsCallback&& UpdateTabletContextsCallback, IStylusInputEventHandler* EventHandler);
+		~FRealTimeStylusPluginAsync();
+
+		// IUnknown
+		virtual HRESULT QueryInterface(REFIID InterfaceID, void** InterfaceObject) override
+		{
+			if (InterfaceID == IID_IStylusAsyncPlugin || InterfaceID == IID_IUnknown)
+			{
+				*InterfaceObject = this;
+				AddRef();
+				return S_OK;
+			}
+
+			*InterfaceObject = nullptr;
+			return E_NOINTERFACE;
+		}
+
+		virtual ULONG AddRef() override
+		{
+			return ++RefCount;
+		}
+
+		virtual ULONG Release() override
+		{
+			const int32 NewRefCount = --RefCount;
+			if (NewRefCount == 0)
+				delete this;
+
+			return NewRefCount;
+		}
+
+		// IStylusPlugin
+		virtual HRESULT RealTimeStylusEnabled(IRealTimeStylus* RealTimeStylus, ULONG TabletContextIDsCount, const TABLET_CONTEXT_ID* TabletContextIDs) override;
+		virtual HRESULT RealTimeStylusDisabled(IRealTimeStylus* RealTimeStylus, ULONG TabletContextIDsCount, const TABLET_CONTEXT_ID* TabletContextIDs) override;
+		virtual HRESULT StylusDown(IRealTimeStylus* RealTimeStylus, const StylusInfo* StylusInfo, ULONG PropertyCount, LONG* PacketBuffer, LONG** InOutPkt) override;
+		virtual HRESULT StylusUp(IRealTimeStylus* RealTimeStylus, const StylusInfo* StylusInfo, ULONG PropertyCount, LONG* PacketBuffer, LONG** InOutPkt) override;
+		virtual HRESULT TabletAdded(IRealTimeStylus* RealTimeStylus, IInkTablet* Tablet) override;
+		virtual HRESULT TabletRemoved(IRealTimeStylus* RealTimeStylus, LONG TabletIndex) override;
+		virtual HRESULT InAirPackets(IRealTimeStylus* RealTimeStylus, const StylusInfo* StylusInfo, ULONG PacketCount, ULONG PacketBufferLength, LONG* PacketBuffer, ULONG* InOutPacketCount, LONG** InOutPacketBuffer) override;
+		virtual HRESULT Packets(IRealTimeStylus* RealTimeStylus, const StylusInfo* StylusInfo, ULONG PacketCount, ULONG PacketBufferLength, LONG* PacketBuffer, ULONG* InOutPacketCount, LONG** InOutPacketBuffer) override;
+		virtual HRESULT StylusInRange(IRealTimeStylus* RealTimeStylus, TABLET_CONTEXT_ID TabletContextID, STYLUS_ID StylusID) override;
+		virtual HRESULT StylusOutOfRange(IRealTimeStylus* RealTimeStylus, TABLET_CONTEXT_ID TabletContextID, STYLUS_ID StylusID) override;
+		virtual HRESULT StylusButtonDown(IRealTimeStylus* RealTimeStylus, STYLUS_ID StylusID, const GUID* GuidStylusButton, POINT* StylusPos) override;
+		virtual HRESULT StylusButtonUp(IRealTimeStylus* RealTimeStylus, STYLUS_ID StylusID, const GUID* GuidStylusButton, POINT* StylusPos) override;
+		virtual HRESULT CustomStylusDataAdded(IRealTimeStylus* RealTimeStylus, const GUID* GuidId, ULONG DataCount, const BYTE* Data) override;
+		virtual HRESULT SystemEvent(IRealTimeStylus* RealTimeStylus, TABLET_CONTEXT_ID TabletContextID, STYLUS_ID StylusID, SYSTEM_EVENT Event, SYSTEM_EVENT_DATA EventData) override;
+		virtual HRESULT Error(IRealTimeStylus* RealTimeStylus, IStylusPlugin* Plugin, RealTimeStylusDataInterest DataInterest, HRESULT ErrorCode, LONG_PTR* InternalKey) override;
+		virtual HRESULT UpdateMapping(IRealTimeStylus* RealTimeStylus) override;
+		virtual HRESULT DataInterest(RealTimeStylusDataInterest* DataInterest) override;
+
+	protected:
+		virtual FString GetName() const override { return "OnGameThread"; }
+
+	private:
+		int32 RefCount = 1;
+	};
+}
+
+#endif
diff --git a/Engine/Plugins/Editor/StylusInput/Source/StylusInput/Private/Windows/RealTimeStylus/RealTimeStylusPluginBase.cpp b/Engine/Plugins/Editor/StylusInput/Source/StylusInput/Private/Windows/RealTimeStylus/RealTimeStylusPluginBase.cpp
new file mode 100644
--- /dev/null
+++ b/Engine/Plugins/Editor/StylusInput/Source/StylusInput/Private/Windows/RealTimeStylus/RealTimeStylusPluginBase.cpp
@@ -0,0 +1,771 @@
+﻿// Copyright Epic Games, Inc. All Rights Reserved.
+
+#include "RealTimeStylusPluginBase.h"
+
+#if PLATFORM_WINDOWS
+
+#include <StylusInput.h>
+#include <StylusInputUtils.h>
+
+#include "RealTimeStylusAPI.h"
+#include "RealTimeStylusUtils.h"
+
+#include <Windows/AllowWindowsPlatformAtomics.h>
+	#include <comdef.h>
+#include <Windows/HideWindowsPlatformAtomics.h>
+
+#define LOG_PREAMBLE "RealTimeStylusPluginBase"
+
+namespace UE::StylusInput::Private::RealTimeStylus
+{
+	bool SetupTabletContextMetadata(IInkTablet& InkTablet, FTabletContext& TabletContext)
+	{
+		const FRealTimeStylusAPI& RealTimeStylusAPI = FRealTimeStylusAPI::GetInstance();
+
+		bool bSuccess = true;
+
+		BSTR String = nullptr;
+		if (Succeeded(InkTablet.get_Name(&String), LOG_PREAMBLE))
+		{
+			TabletContext.Name = String;
+			RealTimeStylusAPI.SysFreeString(String);
+		}
+		else
+		{
+			LogError(LOG_PREAMBLE, FString::Format(TEXT("Could not get name for TabletContext with ID {0}."), {TabletContext.ID}));
+			bSuccess = false;
+		}
+
+		if (Succeeded(InkTablet.get_PlugAndPlayId(&String), LOG_PREAMBLE))
+		{
+			TabletContext.PlugAndPlayID = FString(String);
+			RealTimeStylusAPI.SysFreeString(String);
+		}
+		else
+		{
+			LogError(LOG_PREAMBLE, FString::Format(TEXT("Could not get plug and play ID for TabletContext with ID {0}."), {TabletContext.ID}));
+			bSuccess = false;
+		}
+
+		IInkRectangle* Rectangle = nullptr;
+		if (Succeeded(InkTablet.get_MaximumInputRectangle(&Rectangle), LOG_PREAMBLE))
+		{
+			long Top, Left, Bottom, Right;
+			if (Succeeded(Rectangle->GetRectangle(&Top, &Left, &Bottom, &Right), LOG_PREAMBLE))
+			{
+				TabletContext.InputRectangle = {Left, Top, Right, Bottom};
+			}
+		}
+		else
+		{
+			LogError(LOG_PREAMBLE, FString::Format(TEXT("Could not get input rectangle for TabletContext with ID {0}."), {TabletContext.ID}));
+			bSuccess = false;
+		}
+
+		TabletHardwareCapabilities HardwareCapabilities;
+		if (Succeeded(InkTablet.get_HardwareCapabilities(&HardwareCapabilities), LOG_PREAMBLE))
+		{
+			TabletContext.HardwareCapabilities =
+				static_cast<ETabletHardwareCapabilities>(HardwareCapabilities & THWC_Integrated) |
+				static_cast<ETabletHardwareCapabilities>(HardwareCapabilities & THWC_CursorMustTouch) |
+				static_cast<ETabletHardwareCapabilities>(HardwareCapabilities & THWC_HardProximity) |
+				static_cast<ETabletHardwareCapabilities>(HardwareCapabilities & THWC_CursorsHavePhysicalIds);
+		}
+		else
+		{
+			LogError(LOG_PREAMBLE, FString::Format(TEXT("Could not get hardware capabilities for TabletContext with ID {0}."), {TabletContext.ID}));
+			bSuccess = false;
+		}
+
+		return bSuccess;
+	}
+
+	// TODO compact number of functions needed to only have differences in behavior, and then use offsetof() to apply to different members
+	// TODO XY validate mapping 
+
+	void SetPropertyX(FStylusInputPacket& Packet, const int32 Value, const int8* Data)
+	{
+		const double& Scale = reinterpret_cast<const double*>(Data)[0];
+		const double& Maximum = reinterpret_cast<const double*>(Data)[1];
+		const double& WindowWidth = reinterpret_cast<const double*>(Data)[2];
+
+		Packet.X = Value / Scale / Maximum * WindowWidth;
+	}
+	static_assert(std::is_same_v<FPacketPropertyHandler::FFuncSetProperty, decltype(SetPropertyX)>);
+
+	void SetPropertyY(FStylusInputPacket& Packet, const int32 Value, const int8* Data)
+	{
+		const double& Scale = reinterpret_cast<const double*>(Data)[0];
+		const double& Maximum = reinterpret_cast<const double*>(Data)[1];
+		const double& WindowHeight = reinterpret_cast<const double*>(Data)[2];
+
+		Packet.Y = Value / Scale / Maximum * WindowHeight;
+	}
+	static_assert(std::is_same_v<FPacketPropertyHandler::FFuncSetProperty, decltype(SetPropertyY)>);
+
+	void SetPropertyZ(FStylusInputPacket& Packet, const int32 Value, const int8* Data)
+	{
+		const int32& Minimum = reinterpret_cast<const int32*>(Data)[0];
+		const int32& Maximum = reinterpret_cast<const int32*>(Data)[1];
+
+		Packet.Z = static_cast<float>(Value - Minimum) / (Maximum - Minimum);
+	}
+	static_assert(std::is_same_v<FPacketPropertyHandler::FFuncSetProperty, decltype(SetPropertyZ)>);
+
+	void SetPropertyPacketStatus(FStylusInputPacket& Packet, const int32 Value, const int8*)
+	{
+		Packet.PenStatus = static_cast<EPenStatus>(Value);
+	}
+	static_assert(std::is_same_v<FPacketPropertyHandler::FFuncSetProperty, decltype(SetPropertyPacketStatus)>);
+
+	void SetPropertyTimerTick(FStylusInputPacket& Packet, const int32 Value, const int8* Data)
+	{
+#if 0
+		Packet.TimerTick = Value; // The value provided from the RealTimeStylus API seems to only be non-zero if the stylus is down.
+#else
+		FRealTimeStylusAPI::FFuncGetTickCount* const FuncGetTickCount =	*reinterpret_cast<decltype(FuncGetTickCount)*>(Data);
+		Packet.TimerTick = FuncGetTickCount();
+#endif
+	}
+	static_assert(std::is_same_v<FPacketPropertyHandler::FFuncSetProperty, decltype(SetPropertyTimerTick)>);
+
+	void SetPropertySerialNumber(FStylusInputPacket& Packet, const int32 Value, const int8* Data)
+	{
+		Packet.SerialNumber = Value;
+	}
+	static_assert(std::is_same_v<FPacketPropertyHandler::FFuncSetProperty, decltype(SetPropertySerialNumber)>);
+
+	void SetPropertyNormalPressure(FStylusInputPacket& Packet, const int32 Value, const int8* Data)
+	{
+		const int32& Minimum = reinterpret_cast<const int32*>(Data)[0];
+		const int32& Maximum = reinterpret_cast<const int32*>(Data)[1];
+
+		Packet.NormalPressure = static_cast<float>(Value - Minimum) / (Maximum - Minimum);
+	}
+	static_assert(std::is_same_v<FPacketPropertyHandler::FFuncSetProperty, decltype(SetPropertyNormalPressure)>);
+
+	void SetPropertyTangentPressure(FStylusInputPacket& Packet, const int32 Value, const int8* Data)
+	{
+		const int32& Minimum = reinterpret_cast<const int32*>(Data)[0];
+		const int32& Maximum = reinterpret_cast<const int32*>(Data)[1];
+
+		Packet.TangentPressure = static_cast<float>(Value - Minimum) / (Maximum - Minimum);
+	}
+	static_assert(std::is_same_v<FPacketPropertyHandler::FFuncSetProperty, decltype(SetPropertyTangentPressure)>);
+
+	void SetPropertyButtonPressure(FStylusInputPacket& Packet, const int32 Value, const int8* Data)
+	{
+		const int32& Minimum = reinterpret_cast<const int32*>(Data)[0];
+		const int32& Maximum = reinterpret_cast<const int32*>(Data)[1];
+
+		Packet.ButtonPressure = static_cast<float>(Value - Minimum) / (Maximum - Minimum);
+	}
+	static_assert(std::is_same_v<FPacketPropertyHandler::FFuncSetProperty, decltype(SetPropertyButtonPressure)>);
+
+	void SetPropertyXTiltOrientation(FStylusInputPacket& Packet, const int32 Value, const int8* Data)
+	{
+		const float& Resolution = reinterpret_cast<const float*>(Data)[0];
+
+		Packet.XTiltOrientation = Value / Resolution;
+	}
+	static_assert(std::is_same_v<FPacketPropertyHandler::FFuncSetProperty, decltype(SetPropertyXTiltOrientation)>);
+
+	void SetPropertyYTiltOrientation(FStylusInputPacket& Packet, const int32 Value, const int8* Data)
+	{
+		const float& Resolution = reinterpret_cast<const float*>(Data)[0];
+
+		Packet.YTiltOrientation = Value / Resolution;
+	}
+	static_assert(std::is_same_v<FPacketPropertyHandler::FFuncSetProperty, decltype(SetPropertyYTiltOrientation)>);
+
+	void SetPropertyAzimuthOrientation(FStylusInputPacket& Packet, const int32 Value, const int8* Data)
+	{
+		const float& Resolution = reinterpret_cast<const float*>(Data)[0];
+
+		Packet.AzimuthOrientation = Value / Resolution;
+	}
+	static_assert(std::is_same_v<FPacketPropertyHandler::FFuncSetProperty, decltype(SetPropertyAzimuthOrientation)>);
+
+	void SetPropertyAltitudeOrientation(FStylusInputPacket& Packet, const int32 Value, const int8* Data)
+	{
+		const float& Resolution = reinterpret_cast<const float*>(Data)[0];
+
+		Packet.AltitudeOrientation = Value / Resolution;
+	}
+	static_assert(std::is_same_v<FPacketPropertyHandler::FFuncSetProperty, decltype(SetPropertyAltitudeOrientation)>);
+
+	void SetPropertyTwistOrientation(FStylusInputPacket& Packet, const int32 Value, const int8* Data)
+	{
+		const float& Resolution = reinterpret_cast<const float*>(Data)[0];
+
+		Packet.TwistOrientation = Value / Resolution;
+	}
+	static_assert(std::is_same_v<FPacketPropertyHandler::FFuncSetProperty, decltype(SetPropertyTwistOrientation)>);
+
+	void SetPropertyPitchRotation(FStylusInputPacket& Packet, const int32 Value, const int8* Data)
+	{
+		const float& Resolution = reinterpret_cast<const float*>(Data)[0];
+
+		Packet.PitchRotation = Value / Resolution;
+	}
+	static_assert(std::is_same_v<FPacketPropertyHandler::FFuncSetProperty, decltype(SetPropertyPitchRotation)>);
+
+	void SetPropertyRollRotation(FStylusInputPacket& Packet, const int32 Value, const int8* Data)
+	{
+		const float& Resolution = reinterpret_cast<const float*>(Data)[0];
+
+		Packet.RollRotation = Value / Resolution;
+	}
+	static_assert(std::is_same_v<FPacketPropertyHandler::FFuncSetProperty, decltype(SetPropertyRollRotation)>);
+
+	void SetPropertyYawRotation(FStylusInputPacket& Packet, const int32 Value, const int8* Data)
+	{
+		const float& Resolution = reinterpret_cast<const float*>(Data)[0];
+
+		Packet.YawRotation = Value / Resolution;
+	}
+	static_assert(std::is_same_v<FPacketPropertyHandler::FFuncSetProperty, decltype(SetPropertyYawRotation)>);
+
+	void SetPropertyWidth(FStylusInputPacket& Packet, const int32 Value, const int8* Data)
+	{
+		const float& Resolution = reinterpret_cast<const float*>(Data)[0];
+
+		Packet.Width = Value / Resolution;
+	}
+	static_assert(std::is_same_v<FPacketPropertyHandler::FFuncSetProperty, decltype(SetPropertyWidth)>);
+
+	void SetPropertyHeight(FStylusInputPacket& Packet, const int32 Value, const int8* Data)
+	{
+		const float& Resolution = reinterpret_cast<const float*>(Data)[0];
+
+		Packet.Height = Value / Resolution;
+	}
+	static_assert(std::is_same_v<FPacketPropertyHandler::FFuncSetProperty, decltype(SetPropertyHeight)>);
+
+	void SetPropertyFingerContactConfidence(FStylusInputPacket& Packet, const int32 Value, const int8* Data)
+	{
+		const int32& Minimum = reinterpret_cast<const int32*>(Data)[0];
+		const int32& Maximum = reinterpret_cast<const int32*>(Data)[1];
+
+		Packet.FingerContactConfidence = static_cast<float>(Value - Minimum) / (Maximum - Minimum);
+	}
+	static_assert(std::is_same_v<FPacketPropertyHandler::FFuncSetProperty, decltype(SetPropertyFingerContactConfidence)>);
+
+	void SetPropertyDeviceContactID(FStylusInputPacket& Packet, const int32 Value, const int8* Data)
+	{
+		Packet.DeviceContactID = Value;
+	}
+	static_assert(std::is_same_v<FPacketPropertyHandler::FFuncSetProperty, decltype(SetPropertyDeviceContactID)>);
+
+	FPacketPropertyHandler::FFuncSetProperty* GetSetPropertyFunc(const EPacketPropertyType Type)
+	{
+		switch (Type)
+		{
+		case EPacketPropertyType::X:
+			return &SetPropertyX;
+		case EPacketPropertyType::Y:
+			return &SetPropertyY;
+		case EPacketPropertyType::Z:
+			return &SetPropertyZ;
+		case EPacketPropertyType::PacketStatus:
+			return &SetPropertyPacketStatus;
+		case EPacketPropertyType::TimerTick:
+			return &SetPropertyTimerTick;
+		case EPacketPropertyType::SerialNumber:
+			return &SetPropertySerialNumber;
+		case EPacketPropertyType::NormalPressure:
+			return &SetPropertyNormalPressure;
+		case EPacketPropertyType::TangentPressure:
+			return &SetPropertyTangentPressure;
+		case EPacketPropertyType::ButtonPressure:
+			return &SetPropertyButtonPressure;
+		case EPacketPropertyType::XTiltOrientation:
+			return &SetPropertyXTiltOrientation;
+		case EPacketPropertyType::YTiltOrientation:
+			return &SetPropertyYTiltOrientation;
+		case EPacketPropertyType::AzimuthOrientation:
+			return &SetPropertyAzimuthOrientation;
+		case EPacketPropertyType::AltitudeOrientation:
+			return &SetPropertyAltitudeOrientation;
+		case EPacketPropertyType::TwistOrientation:
+			return &SetPropertyTwistOrientation;
+		case EPacketPropertyType::PitchRotation:
+			return &SetPropertyPitchRotation;
+		case EPacketPropertyType::RollRotation:
+			return &SetPropertyRollRotation;
+		case EPacketPropertyType::YawRotation:
+			return &SetPropertyYawRotation;
+		case EPacketPropertyType::Width:
+			return &SetPropertyWidth;
+		case EPacketPropertyType::Height:
+			return &SetPropertyHeight;
+		case EPacketPropertyType::FingerContactConfidence:
+			return &SetPropertyFingerContactConfidence;
+		case EPacketPropertyType::DeviceContactID:
+			return &SetPropertyDeviceContactID;
+		default:
+			checkSlow(false)
+			return nullptr;
+		}
+	}
+
+	template <int32 DataBufferLength>
+	void AssignSetPropertyData(const EPacketPropertyType Type, int8 *const Data, const FWindowContext& WindowContext, const FPacketProperty& Description)
+	{
+		int8* DataCurrent = Data;
+
+		auto Add = [&DataCurrent]<typename T>(T Value)
+		{
+			*reinterpret_cast<T*>(DataCurrent) = Value;
+			DataCurrent += sizeof(T);
+		};
+
+		switch (Type)
+		{
+		case EPacketPropertyType::X:
+			Add(WindowContext.XYScale.X);
+			Add(WindowContext.XYMaximum.X);
+			Add(WindowContext.WindowSize.X);
+			break;
+		case EPacketPropertyType::Y:
+			Add(WindowContext.XYScale.Y);
+			Add(WindowContext.XYMaximum.Y);
+			Add(WindowContext.WindowSize.Y);
+			break;
+		case EPacketPropertyType::Z:
+		case EPacketPropertyType::NormalPressure:
+		case EPacketPropertyType::TangentPressure:
+		case EPacketPropertyType::ButtonPressure:
+		case EPacketPropertyType::FingerContactConfidence:
+			Add(Description.Minimum);
+			Add(Description.Maximum);
+			break;
+		case EPacketPropertyType::XTiltOrientation:
+		case EPacketPropertyType::YTiltOrientation:
+		case EPacketPropertyType::AzimuthOrientation:
+		case EPacketPropertyType::AltitudeOrientation:
+		case EPacketPropertyType::TwistOrientation:
+		case EPacketPropertyType::PitchRotation:
+		case EPacketPropertyType::RollRotation:
+		case EPacketPropertyType::YawRotation:
+		case EPacketPropertyType::Width:
+		case EPacketPropertyType::Height:
+			Add(Description.Resolution);
+			break;
+		case EPacketPropertyType::TimerTick:
+			Add(FRealTimeStylusAPI::GetInstance().GetTickCount);
+			break;
+		case EPacketPropertyType::PacketStatus:
+		case EPacketPropertyType::SerialNumber:
+		case EPacketPropertyType::DeviceContactID:
+			// do nothing
+			break;
+		default:
+			checkSlow(false)
+		}
+
+		check(std::distance(Data, DataCurrent) <= DataBufferLength)
+	}
+
+	bool SetupTabletContextPacketDescriptionData(IRealTimeStylus* RealTimeStylus, const FWindowContext& WindowContext, FTabletContext& TabletContext)
+	{
+		const FRealTimeStylusAPI& RealTimeStylusAPI = FRealTimeStylusAPI::GetInstance();
+
+		auto InvalidatePacketDescriptions = [&PacketDescriptions = TabletContext.PacketProperties](const int32 StartingIndex)
+		{
+			for (int32 Index = StartingIndex; Index < static_cast<int32>(EPacketPropertyType::Num_Enumerators); ++Index)
+			{
+				PacketDescriptions[Index] = {};
+			}
+		};
+
+		FLOAT InkToDeviceScaleX = 0.0f;
+		FLOAT InkToDeviceScaleY = 0.0f;
+		ULONG PropertiesCount = 0;
+		PACKET_PROPERTY* Properties = nullptr;
+		if (Succeeded(RealTimeStylus->GetPacketDescriptionData(TabletContext.ID, &InkToDeviceScaleX, &InkToDeviceScaleY, &PropertiesCount, &Properties), LOG_PREAMBLE))
+		{
+			if (!Properties)
+			{
+				LogWarning(LOG_PREAMBLE, FString::Format(
+					           TEXT("Retrieved nullptr when trying to get packet description data for TabletContext with ID {0}."), {TabletContext.ID}));
+				return false;
+			}
+
+			if (PropertiesCount == 0)
+			{
+				LogWarning(LOG_PREAMBLE, FString::Format(
+					           TEXT("Retrieved zero packets properties when trying to get packet description data for TabletContext with ID {0}."),
+					           {TabletContext.ID}));
+				return false;
+			}
+
+			for (ULONG Index = 0; Index < PropertiesCount; ++Index)
+			{
+#pragma warning(push)
+#pragma warning(disable : 28199)
+				const PACKET_PROPERTY& Property = Properties[Index];
+#pragma warning(pop)
+
+				bool bFoundProperty = false;
+				for (const FPacketPropertyConstant& PropertyConstant : PacketPropertyConstants)
+				{
+					if (Property.guid == PropertyConstant.Guid)
+					{
+						const ETabletSupportedProperties SupportedProperty = static_cast<ETabletSupportedProperties>(
+							1 << static_cast<std::underlying_type_t<ETabletSupportedProperties>>(PropertyConstant.PacketPropertyType)
+						);
+						TabletContext.SupportedProperties = TabletContext.SupportedProperties | SupportedProperty;
+
+						FPacketProperty& PropertyProperty = TabletContext.PacketProperties[Index];
+						PropertyProperty.Type = static_cast<EPacketPropertyType>(PropertyConstant.PacketPropertyType);
+						PropertyProperty.Minimum = Property.PropertyMetrics.nLogicalMin;
+						PropertyProperty.Maximum = Property.PropertyMetrics.nLogicalMax;
+						PropertyProperty.MetricUnit = static_cast<ETabletPropertyMetricUnit>(Property.PropertyMetrics.Units);
+						PropertyProperty.Resolution = Property.PropertyMetrics.fResolution;
+
+						FPacketPropertyHandler& PropertyPropertyHandler = TabletContext.PacketPropertyHandlers[Index];
+						PropertyPropertyHandler.SetProperty = GetSetPropertyFunc(PropertyConstant.PacketPropertyType);
+						AssignSetPropertyData<FPacketPropertyHandler::SetPropertyDataBufferLength>(
+							PropertyConstant.PacketPropertyType, PropertyPropertyHandler.SetPropertyData, WindowContext, PropertyProperty);
+
+						if (static_cast<int8>(PropertyProperty.MetricUnit) >= static_cast<int8>(ETabletPropertyMetricUnit::Num_Enumerators))
+						{
+							PropertyProperty.MetricUnit = ETabletPropertyMetricUnit::Default;
+
+							LogWarning(LOG_PREAMBLE, FString::Format(
+									   TEXT("Encountered unknown metric unit value '{0}' while evaluating packet description data for TabletContext with ID {1}."),
+									   {static_cast<int8>(PropertyProperty.MetricUnit), TabletContext.ID}));
+						}
+
+						bFoundProperty = true;
+						break;
+					}
+				}
+
+				if (!bFoundProperty)
+				{
+					wchar_t GUIDStringBuffer[64];
+					const int32 GUIDStringLength = RealTimeStylusAPI.StringFromGUID2(Property.guid, GUIDStringBuffer, std::size(GUIDStringBuffer));
+					if (ensure(GUIDStringLength > 0))
+					{
+						LogWarning(LOG_PREAMBLE, FString::Format(
+							           TEXT("Encountered unknown property '{0}' while evaluating packet description data for TabletContext with ID {1}."),
+							           {GUIDStringBuffer, TabletContext.ID}));
+					}
+				}
+			}
+
+			RealTimeStylusAPI.CoTaskMemFree(Properties);
+
+			InvalidatePacketDescriptions(PropertiesCount);
+		}
+		else
+		{
+			LogError(LOG_PREAMBLE, FString::Format(TEXT("Could not get packet description data for TabletContext with ID {0}."), {TabletContext.ID}));
+
+			InvalidatePacketDescriptions(0);
+
+			return false;
+		}
+
+		return true;
+	}
+
+	bool SetupTabletContext(IRealTimeStylus* RealTimeStylus, const FWindowContext& WindowContext, FTabletContext& TabletContext)
+	{
+		IInkTablet* InkTablet = nullptr;
+
+		if (Failed(RealTimeStylus->GetTabletFromTabletContextId(TabletContext.ID, &InkTablet), LOG_PREAMBLE))
+		{
+			LogError(LOG_PREAMBLE, FString::Format(TEXT("Could not get tablet context data for ID {0}."), {TabletContext.ID}));
+			return false;
+		}
+
+		bool bSuccess = true;
+
+		bSuccess &= SetupTabletContextMetadata(*InkTablet, TabletContext);
+		bSuccess &= SetupTabletContextPacketDescriptionData(RealTimeStylus, WindowContext, TabletContext);
+
+		return bSuccess;
+	}
+
+	bool FRealTimeStylusPluginBase::AddEventHandler(IStylusInputEventHandler* EventHandler)
+	{
+		check(EventHandler)
+
+		if (EventHandlers.Contains(EventHandler))
+		{
+			LogWarning(LOG_PREAMBLE, FString::Format(TEXT("Event handler '{0}' already exists in {1} plugin."), {EventHandler->GetName(), GetName()}));
+			return false;
+		}
+
+		EventHandlers.Add(EventHandler);
+
+		LogVerbose(LOG_PREAMBLE, FString::Format(TEXT("Event handler '{0}' was added to {1} plugin."), {EventHandler->GetName(), GetName()}));
+
+		return true;
+	}
+
+	bool FRealTimeStylusPluginBase::RemoveEventHandler(IStylusInputEventHandler* EventHandler)
+	{
+		check(EventHandler)
+
+		const bool bWasRemoved = EventHandlers.Remove(EventHandler) > 0;
+
+		if (bWasRemoved)
+		{
+			LogVerbose(LOG_PREAMBLE, FString::Format(TEXT("Event handler '{0}' was removed from {1} plugin."), {EventHandler->GetName(), GetName()}));
+		}
+
+		return bWasRemoved;
+	}
+
+	FRealTimeStylusPluginBase::FRealTimeStylusPluginBase(IStylusInputInstance* Instance,
+	                                                             FGetWindowContextCallback&& GetWindowContextCallback,
+	                                                             FUpdateTabletContextsCallback&& UpdateTabletContextsCallback)
+		: Instance(Instance)
+		, GetWindowContextCallback(MoveTemp(GetWindowContextCallback))
+		, UpdateTabletContextsCallback(MoveTemp(UpdateTabletContextsCallback)) 
+	{
+		check(this->GetWindowContextCallback.IsBound())
+		check(this->UpdateTabletContextsCallback.IsBound())
+	}
+
+	void FRealTimeStylusPluginBase::DebugEvent(const FString& Message) const
+	{
+		for (IStylusInputEventHandler* EventHandler : EventHandlers)
+		{
+			EventHandler->OnDebugEvent(Message, Instance);
+		}
+	}
+
+	HRESULT FRealTimeStylusPluginBase::ProcessDataInterest(RealTimeStylusDataInterest* DataInterest)
+	{
+		constexpr bool bGetAllData = false;
+
+		if (bGetAllData)
+		{
+			*DataInterest = RTSDI_AllData;
+
+			DebugEvent("Requested all data from stylus input.");
+		}
+		else
+		{
+			*DataInterest = static_cast<RealTimeStylusDataInterest>(
+				RTSDI_Error |
+				RTSDI_RealTimeStylusEnabled | RTSDI_RealTimeStylusDisabled |
+				RTSDI_InAirPackets | RTSDI_Packets |
+				RTSDI_StylusDown | RTSDI_StylusUp);
+
+			DebugEvent("Requested data for Error, RealTimeStylusEnabled, RealTimeStylusDisabled, InAirPackets, Packets, StylusDown, StylusUp "
+				"from stylus input (DataInterest event).");
+		}
+
+		return S_OK;
+	}
+
+	HRESULT FRealTimeStylusPluginBase::ProcessError(RealTimeStylusDataInterest DataInterest, const HRESULT ErrorCode)
+	{
+		constexpr bool bShowAllErrors = false;
+
+		if (bShowAllErrors || ErrorCode != E_NOTIMPL)
+		{
+			const char* DataInterestStr = [DataInterest]
+			{
+				switch (DataInterest)
+				{
+				case RTSDI_AllData:
+					return "AllData";
+				case RTSDI_None:
+					return "None";
+				case RTSDI_Error:
+					return "Error";
+				case RTSDI_RealTimeStylusEnabled:
+					return "RealTimeStylusEnabled";
+				case RTSDI_RealTimeStylusDisabled:
+					return "RealTimeStylusDisabled";
+				case RTSDI_StylusNew:
+					return "StylusNew";
+				case RTSDI_StylusInRange:
+					return "StylusInRange";
+				case RTSDI_InAirPackets:
+					return "InAirPackets";
+				case RTSDI_StylusOutOfRange:
+					return "StylusOutOfRange";
+				case RTSDI_StylusDown:
+					return "StylusDown";
+				case RTSDI_Packets:
+					return "Packets";
+				case RTSDI_StylusUp:
+					return "StylusUp";
+				case RTSDI_StylusButtonUp:
+					return "StylusButtonUp";
+				case RTSDI_StylusButtonDown:
+					return "StylusButtonDown";
+				case RTSDI_SystemEvents:
+					return "SystemEvents";
+				case RTSDI_TabletAdded:
+					return "TabletAdded";
+				case RTSDI_TabletRemoved:
+					return "TabletRemoved";
+				case RTSDI_CustomStylusDataAdded:
+					return "CustomStylusDataAdded";
+				case RTSDI_UpdateMapping:
+					return "UpdateMapping";
+				case RTSDI_DefaultEvents:
+					return "DefaultEvents";
+				}
+				return "<unknown>";
+			}();
+
+			const FString ErrorMessage = FString::Format(TEXT("Error in {0} plugin: {1}, Error={2} ({3})."), {
+				                                             GetName(), DataInterestStr, _com_error(ErrorCode).ErrorMessage(), static_cast<int32>(ErrorCode)
+			                                             });
+
+			DebugEvent(ErrorMessage);
+		}
+
+		return S_OK;
+	}
+
+	HRESULT FRealTimeStylusPluginBase::ProcessPackets(const StylusInfo* StylusInfo, uint32 PacketCount, uint32 PacketBufferLength, EPacketType Type,
+	                                                  const int32* PacketBuffer)
+	{
+		checkSlow(StylusInfo)
+		checkSlow(PacketBuffer)
+
+		const FPacketPropertyHandler* Handlers = GetPacketPropertyHandlers(StylusInfo->tcid);
+		if (!Handlers)
+		{
+			return E_FAIL;
+		}
+
+		const uint32 PropertyCount = PacketBufferLength / PacketCount;
+		checkSlow(PropertyCount < static_cast<uint32>(EPacketPropertyType::Num_Enumerators));
+
+		for (uint32 PacketIndex = 0; PacketIndex < PacketCount; PacketIndex += PropertyCount)
+		{
+			PacketStats.NewPacket();
+
+			FStylusInputPacket Packet{StylusInfo->tcid, StylusInfo->cid, Type};
+
+			for (uint32 PropertyIndex = 0; PropertyIndex < PropertyCount; ++PropertyIndex)
+			{
+				const FPacketPropertyHandler& Handler = Handlers[PropertyIndex];
+				checkSlow(Handler.SetProperty)
+				Handler.SetProperty(Packet, PacketBuffer[PropertyIndex], Handler.SetPropertyData);
+			}
+
+			for (IStylusInputEventHandler* EventHandler : EventHandlers)
+			{
+				EventHandler->OnPacket(Packet, Instance);
+			}
+
+			PacketBuffer += PropertyCount;
+		}
+
+		return S_OK;
+	}
+
+	HRESULT FRealTimeStylusPluginBase::ProcessRealTimeStylusEnabled(IRealTimeStylus* RealTimeStylus, const uint32 TabletContextIDsCount,
+																   const TABLET_CONTEXT_ID* TabletContextIDs)
+	{
+		DebugEvent("Stylus input was enabled (RealTimeStylusEnabled event).");
+
+		return UpdateTabletContexts(RealTimeStylus, TabletContextIDsCount, TabletContextIDs);
+	}
+
+	HRESULT FRealTimeStylusPluginBase::ProcessRealTimeStylusDisabled(IRealTimeStylus* RealTimeStylus, const uint32 TabletContextIDsCount,
+																	const TABLET_CONTEXT_ID* TabletContextIDs)
+	{
+		DebugEvent("Stylus input was disabled (RealTimeStylusDisabled event).");
+
+		return UpdateTabletContexts(RealTimeStylus, TabletContextIDsCount, TabletContextIDs);
+	}
+
+	HRESULT FRealTimeStylusPluginBase::ProcessTabletAdded(IInkTablet* Tablet)
+	{
+		// After a TabletAdded event, a RealTimeStylusDisabled event will fire directly followed by a RealTimeStylusEnabled event.
+		// We are using these two events instead to update the tablet contexts.
+
+		return E_NOTIMPL;
+	}
+
+	HRESULT FRealTimeStylusPluginBase::ProcessTabletRemoved(LONG TabletIndex)
+	{
+		// For simplicity, we don't remove a tablet context when this event is received.
+		// However, since there are no more packets coming through for the removed tablet there should be nothing that's continuing to access the outdated
+		// tablet context data.
+
+		return E_NOTIMPL;
+	}
+
+	const FPacketPropertyHandler* FRealTimeStylusPluginBase::GetPacketPropertyHandlers(uint32 TabletContextID) const
+	{
+		const TSharedPtr<FTabletContext> TabletContext = TabletContexts.Get(TabletContextID);
+		return TabletContext ? TabletContext->PacketPropertyHandlers : nullptr;
+	}
+
+	HRESULT FRealTimeStylusPluginBase::UpdateTabletContexts(IRealTimeStylus* RealTimeStylus, uint32 TabletContextIDsNum, const TABLET_CONTEXT_ID* TabletContextIDs)
+	{
+		const FWindowContext& WindowContext = GetWindowContextCallback.Execute();
+
+		bool bSuccess = true;
+
+		// Remove outdated tablet contexts
+		uint32 TabletContextIndex = 0;
+		while (TabletContextIndex < TabletContexts.Num())
+		{
+			const uint32 ExistingTabletContextID = TabletContexts[TabletContextIndex]->ID;
+
+			const bool bContainsID = [ExistingTabletContextID, TabletContextIDsNum, TabletContextIDs]
+			{
+				for (uint32 TabletContextIDsIndex = 0; TabletContextIDsIndex < TabletContextIDsNum; ++TabletContextIDsIndex)
+				{
+					if (ExistingTabletContextID == TabletContextIDs[TabletContextIDsIndex])
+					{
+						return true;
+					}
+				}
+				return false;
+			}();
+
+			if (bContainsID)
+			{
+				++TabletContextIndex;
+			}
+			else
+			{
+				TabletContexts.Remove(ExistingTabletContextID);
+
+				LogVerbose(LOG_PREAMBLE, FString::Format(TEXT("Removed tablet context data for ID {0}."), {ExistingTabletContextID}));
+			}
+		}
+
+		// Add new tablet contexts
+		for (uint32 TabletContextIDsIndex = 0; TabletContextIDsIndex < TabletContextIDsNum; ++TabletContextIDsIndex)
+		{
+			const uint32 TabletContextID = TabletContextIDs[TabletContextIDsIndex];
+			if (!TabletContexts.Contains(TabletContextID))
+			{
+				const TSharedRef<FTabletContext> TabletContextRef = TabletContexts.Add(TabletContextID);
+				FTabletContext& TabletContext = TabletContextRef.Get();
+
+				bSuccess &= SetupTabletContext(RealTimeStylus, WindowContext, TabletContext);
+
+				LogVerbose(LOG_PREAMBLE, FString::Format(TEXT("Added TabletContext for ID {0} [{1}, {2}]."), {
+															TabletContext.ID, TabletContext.Name, TabletContext.PlugAndPlayID
+														}));
+			}
+		}
+
+		UpdateTabletContextsCallback.Execute(TabletContexts);
+
+		return bSuccess ? S_OK : E_FAIL;
+	}
+}
+
+#undef LOG_PREAMBLE
+
+#endif
diff --git a/Engine/Plugins/Editor/StylusInput/Source/StylusInput/Private/Windows/RealTimeStylus/RealTimeStylusPluginBase.h b/Engine/Plugins/Editor/StylusInput/Source/StylusInput/Private/Windows/RealTimeStylus/RealTimeStylusPluginBase.h
new file mode 100644
--- /dev/null
+++ b/Engine/Plugins/Editor/StylusInput/Source/StylusInput/Private/Windows/RealTimeStylus/RealTimeStylusPluginBase.h
@@ -0,0 +1,72 @@
+﻿// Copyright Epic Games, Inc. All Rights Reserved.
+
+#pragma once
+
+#include <HAL/Platform.h>
+
+#if PLATFORM_WINDOWS
+
+#include <StylusInputPacket.h>
+#include <Delegates/Delegate.h>
+
+#include "RealTimeStylusStats.h"
+#include "RealTimeStylusTabletContext.h"
+
+#include <Windows/AllowWindowsPlatformTypes.h>
+	#include <Microsoft/COMPointer.h>
+	#include <RTSCom.h>
+#include <Windows/HideWindowsPlatformTypes.h>
+
+namespace UE::StylusInput
+{
+	class IStylusInputEventHandler;
+	class IStylusInputInstance;
+}
+
+namespace UE::StylusInput::Private::RealTimeStylus
+{
+	class FRealTimeStylusAPI;
+
+	DECLARE_DELEGATE_RetVal(const FWindowContext&, FGetWindowContextCallback);
+	DECLARE_DELEGATE_OneParam(FUpdateTabletContextsCallback, const FTabletContextContainer&);
+
+	class FRealTimeStylusPluginBase
+	{
+	public:
+		bool AddEventHandler(IStylusInputEventHandler* EventHandler);
+		bool RemoveEventHandler(IStylusInputEventHandler* EventHandler);
+		int32 NumEventHandlers() const { return EventHandlers.Num(); }
+
+		float GetPacketsPerSecond() const { return PacketStats.GetPacketsPerSecond(); }
+
+	protected:
+		FRealTimeStylusPluginBase(IStylusInputInstance* Instance, FGetWindowContextCallback&& GetWindowContextCallback,
+		                          FUpdateTabletContextsCallback&& UpdateTabletContextsCallback);
+		~FRealTimeStylusPluginBase() = default;
+
+		virtual FString GetName() const = 0;
+
+		void DebugEvent(const FString& Message) const;
+
+		HRESULT ProcessDataInterest(RealTimeStylusDataInterest* DataInterest);
+		HRESULT ProcessError(RealTimeStylusDataInterest DataInterest, HRESULT ErrorCode);
+		HRESULT ProcessPackets(const StylusInfo* StylusInfo, uint32 PacketCount, uint32 PacketBufferLength, EPacketType Type, const int32* PacketBuffer);
+		HRESULT ProcessRealTimeStylusEnabled(IRealTimeStylus* RealTimeStylus, uint32 TabletContextIDsCount, const TABLET_CONTEXT_ID* TabletContextIDs);
+		HRESULT ProcessRealTimeStylusDisabled(IRealTimeStylus* RealTimeStylus, uint32 TabletContextIDsCount, const TABLET_CONTEXT_ID* TabletContextIDs);
+		HRESULT ProcessTabletAdded(IInkTablet* Tablet);
+		HRESULT ProcessTabletRemoved(LONG TabletIndex);
+
+	private:
+		const FPacketPropertyHandler* GetPacketPropertyHandlers(uint32 TabletContextID) const;
+		HRESULT UpdateTabletContexts(IRealTimeStylus* RealTimeStylus, uint32 TabletContextIDsNum, const TABLET_CONTEXT_ID* TabletContextIDs);
+
+		IStylusInputInstance *const Instance;
+		FPacketStats PacketStats;
+		FGetWindowContextCallback GetWindowContextCallback;
+		FUpdateTabletContextsCallback UpdateTabletContextsCallback;
+		FTabletContextContainer TabletContexts;
+		TArray<IStylusInputEventHandler*> EventHandlers;
+	};
+}
+
+#endif
diff --git a/Engine/Plugins/Editor/StylusInput/Source/StylusInput/Private/Windows/RealTimeStylus/RealTimeStylusPluginSync.cpp b/Engine/Plugins/Editor/StylusInput/Source/StylusInput/Private/Windows/RealTimeStylus/RealTimeStylusPluginSync.cpp
new file mode 100644
--- /dev/null
+++ b/Engine/Plugins/Editor/StylusInput/Source/StylusInput/Private/Windows/RealTimeStylus/RealTimeStylusPluginSync.cpp
@@ -0,0 +1,134 @@
+﻿// Copyright Epic Games, Inc. All Rights Reserved.
+
+#include "RealTimeStylusPluginSync.h"
+
+#if PLATFORM_WINDOWS
+
+#include "RealTimeStylusAPI.h"
+
+namespace UE::StylusInput::Private::RealTimeStylus
+{
+	FRealTimeStylusPluginSync::FRealTimeStylusPluginSync(IStylusInputInstance* Instance,
+	                                                     FGetWindowContextCallback&& GetWindowContext,
+	                                                     FUpdateTabletContextsCallback&& UpdateTabletContextsCallback,
+	                                                     IStylusInputEventHandler* EventHandler)
+		: FRealTimeStylusPluginBase(Instance, MoveTemp(GetWindowContext), MoveTemp(UpdateTabletContextsCallback))
+	{
+		if (EventHandler)
+		{
+			// Immediately install an event handler during construction to capture events coming through during plugin initialization.
+			AddEventHandler(EventHandler);
+		}
+	}
+
+	FRealTimeStylusPluginSync::~FRealTimeStylusPluginSync()
+	{
+		if (FreeThreadedMarshaler)
+		{
+			FreeThreadedMarshaler->Release();
+		}
+	}
+
+	HRESULT FRealTimeStylusPluginSync::CreateFreeThreadMarshaler()
+	{
+		check(FreeThreadedMarshaler == nullptr);
+
+		const FRealTimeStylusAPI& RealTimeStylusAPI = FRealTimeStylusAPI::GetInstance();
+		return RealTimeStylusAPI.CoCreateFreeThreadedMarshaler(this, &FreeThreadedMarshaler);
+	}
+
+	HRESULT FRealTimeStylusPluginSync::RealTimeStylusEnabled(IRealTimeStylus* RealTimeStylus, const ULONG TabletContextIDsCount,
+	                                                             const TABLET_CONTEXT_ID* TabletContextIDs)
+	{
+		return ProcessRealTimeStylusEnabled(RealTimeStylus, TabletContextIDsCount, TabletContextIDs);
+	}
+
+	HRESULT FRealTimeStylusPluginSync::RealTimeStylusDisabled(IRealTimeStylus* RealTimeStylus, const ULONG TabletContextIDsCount,
+	                                                              const TABLET_CONTEXT_ID* TabletContextIDs)
+	{
+		return ProcessRealTimeStylusDisabled(RealTimeStylus, TabletContextIDsCount, TabletContextIDs);
+	}
+
+	HRESULT FRealTimeStylusPluginSync::StylusDown(IRealTimeStylus* RealTimeStylus, const StylusInfo* StylusInfo, const ULONG PropertyCount,
+	                                                  LONG* PacketBuffer, LONG**)
+	{
+		return ProcessPackets(StylusInfo, 1, PropertyCount, EPacketType::StylusDown, reinterpret_cast<int32*>(PacketBuffer));
+	}
+
+	HRESULT FRealTimeStylusPluginSync::StylusUp(IRealTimeStylus* RealTimeStylus, const StylusInfo* StylusInfo, const ULONG PropertyCount,
+	                                                LONG* PacketBuffer, LONG**)
+	{
+		return ProcessPackets(StylusInfo, 1, PropertyCount, EPacketType::StylusUp, reinterpret_cast<int32*>(PacketBuffer));
+	}
+
+	HRESULT FRealTimeStylusPluginSync::TabletAdded(IRealTimeStylus* RealTimeStylus, IInkTablet* Tablet)
+	{
+		return ProcessTabletAdded(Tablet);
+	}
+
+	HRESULT FRealTimeStylusPluginSync::TabletRemoved(IRealTimeStylus* RealTimeStylus, const LONG TabletIndex)
+	{
+		return ProcessTabletRemoved(TabletIndex);
+	}
+
+	HRESULT FRealTimeStylusPluginSync::InAirPackets(IRealTimeStylus* RealTimeStylus, const StylusInfo* StylusInfo, const ULONG PacketCount,
+	                                                    const ULONG PacketBufferLength, LONG* PacketBuffer, ULONG*, LONG**)
+	{
+		return ProcessPackets(StylusInfo, PacketCount, PacketBufferLength, EPacketType::AboveDigitizer, reinterpret_cast<int32*>(PacketBuffer));
+	}
+
+	HRESULT FRealTimeStylusPluginSync::Packets(IRealTimeStylus* RealTimeStylus, const StylusInfo* StylusInfo, const ULONG PacketCount,
+	                                               const ULONG PacketBufferLength, LONG* PacketBuffer, ULONG*, LONG**)
+	{
+		return ProcessPackets(StylusInfo, PacketCount, PacketBufferLength, EPacketType::OnDigitizer, reinterpret_cast<int32*>(PacketBuffer));
+	}
+
+	HRESULT FRealTimeStylusPluginSync::StylusInRange(IRealTimeStylus* RealTimeStylus, TABLET_CONTEXT_ID TabletContextID, STYLUS_ID StylusID)
+	{
+		return E_NOTIMPL;
+	}
+
+	HRESULT FRealTimeStylusPluginSync::StylusOutOfRange(IRealTimeStylus* RealTimeStylus, TABLET_CONTEXT_ID TabletContextID, STYLUS_ID StylusID)
+	{
+		return E_NOTIMPL;
+	}
+
+	HRESULT FRealTimeStylusPluginSync::StylusButtonDown(IRealTimeStylus* RealTimeStylus, STYLUS_ID StylusID, const GUID* GuidStylusButton, POINT* StylusPos)
+	{
+		return E_NOTIMPL;
+	}
+
+	HRESULT FRealTimeStylusPluginSync::StylusButtonUp(IRealTimeStylus* RealTimeStylus, STYLUS_ID StylusID, const GUID* GuidStylusButton, POINT* StylusPos)
+	{
+		return E_NOTIMPL;
+	}
+
+	HRESULT FRealTimeStylusPluginSync::CustomStylusDataAdded(IRealTimeStylus* RealTimeStylus, const GUID* GuidId, ULONG DataCount, const BYTE* Data)
+	{
+		return E_NOTIMPL;
+	}
+
+	HRESULT FRealTimeStylusPluginSync::SystemEvent(IRealTimeStylus* RealTimeStylus, TABLET_CONTEXT_ID TabletContextID, STYLUS_ID StylusID,
+	                                                   SYSTEM_EVENT Event, SYSTEM_EVENT_DATA EventData)
+	{
+		return E_NOTIMPL;
+	}
+
+	HRESULT FRealTimeStylusPluginSync::Error(IRealTimeStylus*, IStylusPlugin* Plugin, const RealTimeStylusDataInterest DataInterest,
+	                                             const HRESULT ErrorCode, LONG_PTR*)
+	{
+		return ProcessError(DataInterest, ErrorCode);
+	}
+
+	HRESULT FRealTimeStylusPluginSync::UpdateMapping(IRealTimeStylus* RealTimeStylus)
+	{
+		return E_NOTIMPL;
+	}
+
+	HRESULT FRealTimeStylusPluginSync::DataInterest(RealTimeStylusDataInterest* DataInterest)
+	{
+		return ProcessDataInterest(DataInterest);
+	}
+}
+
+#endif
diff --git a/Engine/Plugins/Editor/StylusInput/Source/StylusInput/Private/Windows/RealTimeStylus/RealTimeStylusPluginSync.h b/Engine/Plugins/Editor/StylusInput/Source/StylusInput/Private/Windows/RealTimeStylus/RealTimeStylusPluginSync.h
new file mode 100644
--- /dev/null
+++ b/Engine/Plugins/Editor/StylusInput/Source/StylusInput/Private/Windows/RealTimeStylus/RealTimeStylusPluginSync.h
@@ -0,0 +1,85 @@
+﻿// Copyright Epic Games, Inc. All Rights Reserved.
+
+#pragma once
+
+#include <HAL/Platform.h>
+
+#if PLATFORM_WINDOWS
+
+#include "RealTimeStylusPluginBase.h"
+
+namespace UE::StylusInput::Private::RealTimeStylus
+{
+	class FRealTimeStylusPluginSync final : public IStylusSyncPlugin, public FRealTimeStylusPluginBase
+	{
+	public:
+		FRealTimeStylusPluginSync(IStylusInputInstance* Instance, FGetWindowContextCallback&& GetWindowContext,
+		                              FUpdateTabletContextsCallback&& UpdateTabletContextsCallback, IStylusInputEventHandler* EventHandler);
+		~FRealTimeStylusPluginSync();
+
+		HRESULT CreateFreeThreadMarshaler();
+
+		// IUnknown
+		virtual HRESULT QueryInterface(REFIID InterfaceID, void** InterfaceObject) override
+		{
+			if (InterfaceID == IID_IStylusSyncPlugin || InterfaceID == IID_IUnknown)
+			{
+				*InterfaceObject = this;
+				AddRef();
+				return S_OK;
+			}
+
+			if (InterfaceID == IID_IMarshal && FreeThreadedMarshaler)
+			{
+				return FreeThreadedMarshaler->QueryInterface(InterfaceID, InterfaceObject);
+			}
+
+			*InterfaceObject = nullptr;
+			return E_NOINTERFACE;
+		}
+
+		virtual ULONG AddRef() override
+		{
+			return ++RefCount;
+		}
+
+		virtual ULONG Release() override
+		{
+			const int32 NewRefCount = --RefCount;
+			if (NewRefCount == 0)
+				delete this;
+
+			return NewRefCount;
+		}
+
+		// IStylusPlugin
+		virtual HRESULT RealTimeStylusEnabled(IRealTimeStylus* RealTimeStylus, ULONG TabletContextIDsCount, const TABLET_CONTEXT_ID* TabletContextIDs) override;
+		virtual HRESULT RealTimeStylusDisabled(IRealTimeStylus* RealTimeStylus, ULONG TabletContextIDsCount, const TABLET_CONTEXT_ID* TabletContextIDs) override;
+		virtual HRESULT StylusDown(IRealTimeStylus* RealTimeStylus, const StylusInfo* StylusInfo, ULONG PropertyCount, LONG* PacketBuffer, LONG** InOutPkt) override;
+		virtual HRESULT StylusUp(IRealTimeStylus* RealTimeStylus, const StylusInfo* StylusInfo, ULONG PropertyCount, LONG* PacketBuffer, LONG** InOutPkt) override;
+		virtual HRESULT TabletAdded(IRealTimeStylus* RealTimeStylus, IInkTablet* Tablet) override;
+		virtual HRESULT TabletRemoved(IRealTimeStylus* RealTimeStylus, LONG TabletIndex) override;
+		virtual HRESULT InAirPackets(IRealTimeStylus* RealTimeStylus, const StylusInfo* StylusInfo, ULONG PacketCount, ULONG PacketBufferLength, LONG* PacketBuffer, ULONG* InOutPacketCount, LONG** InOutPacketBuffer) override;
+		virtual HRESULT Packets(IRealTimeStylus* RealTimeStylus, const StylusInfo* StylusInfo, ULONG PacketCount, ULONG PacketBufferLength, LONG* PacketBuffer, ULONG* InOutPacketCount, LONG** InOutPacketBuffer) override;
+		virtual HRESULT StylusInRange(IRealTimeStylus* RealTimeStylus, TABLET_CONTEXT_ID TabletContextID, STYLUS_ID StylusID) override;
+		virtual HRESULT StylusOutOfRange(IRealTimeStylus* RealTimeStylus, TABLET_CONTEXT_ID TabletContextID, STYLUS_ID StylusID) override;
+		virtual HRESULT StylusButtonDown(IRealTimeStylus* RealTimeStylus, STYLUS_ID StylusID, const GUID* GuidStylusButton, POINT* StylusPos) override;
+		virtual HRESULT StylusButtonUp(IRealTimeStylus* RealTimeStylus, STYLUS_ID StylusID, const GUID* GuidStylusButton, POINT* StylusPos) override;
+		virtual HRESULT CustomStylusDataAdded(IRealTimeStylus* RealTimeStylus, const GUID* GuidId, ULONG DataCount, const BYTE* Data) override;
+		virtual HRESULT SystemEvent(IRealTimeStylus* RealTimeStylus, TABLET_CONTEXT_ID TabletContextID, STYLUS_ID StylusID, SYSTEM_EVENT Event, SYSTEM_EVENT_DATA EventData) override;
+		virtual HRESULT Error(IRealTimeStylus* RealTimeStylus, IStylusPlugin* Plugin, RealTimeStylusDataInterest DataInterest, HRESULT ErrorCode, LONG_PTR* InternalKey) override;
+		virtual HRESULT UpdateMapping(IRealTimeStylus* RealTimeStylus) override;
+		virtual HRESULT DataInterest(RealTimeStylusDataInterest* DataInterest) override;
+
+	protected:
+		// FRealTimeStylusPluginBase
+		virtual FString GetName() const override { return "Asynchronous"; }
+
+	private:
+		int32 RefCount = 1;
+
+		IUnknown* FreeThreadedMarshaler = nullptr;
+	};
+}
+
+#endif
diff --git a/Engine/Plugins/Editor/StylusInput/Source/StylusInput/Private/Windows/RealTimeStylus/RealTimeStylusStats.h b/Engine/Plugins/Editor/StylusInput/Source/StylusInput/Private/Windows/RealTimeStylus/RealTimeStylusStats.h
new file mode 100644
--- /dev/null
+++ b/Engine/Plugins/Editor/StylusInput/Source/StylusInput/Private/Windows/RealTimeStylus/RealTimeStylusStats.h
@@ -0,0 +1,46 @@
+﻿// Copyright Epic Games, Inc. All Rights Reserved.
+
+#pragma once
+
+#include <HAL/Platform.h>
+
+#if PLATFORM_WINDOWS
+
+#include <Windows/WindowsPlatformTime.h>
+
+namespace UE::StylusInput::Private::RealTimeStylus
+{
+	class FPacketStats
+	{
+	public:
+		void NewPacket()
+		{
+			const uint32 Timestamp = FPlatformTime::Cycles();
+
+			if (Latest - Earliest >= CyclesPerSecond)
+			{
+				PacketsPerSecond = Earliest != Latest ? NumPackets / ((Latest - Earliest) / CyclesPerSecond) : 0.0f;
+				Earliest = Timestamp;
+				NumPackets = -1;
+			}
+
+			Latest = Timestamp;
+			++NumPackets;
+		}
+
+		float GetPacketsPerSecond() const
+		{
+			const uint32 Timestamp = FPlatformTime::Cycles();
+			return Timestamp - Latest <= CyclesPerSecond ? PacketsPerSecond : 0.0f;
+		}
+
+	private:
+		uint32 Earliest = 0;
+		uint32 Latest = 0;
+		uint32 NumPackets = 0;
+		float PacketsPerSecond = 0.0f;
+		const double CyclesPerSecond = 1.0 / FPlatformTime::GetSecondsPerCycle();
+	};
+}
+
+#endif
diff --git a/Engine/Plugins/Editor/StylusInput/Source/StylusInput/Private/Windows/RealTimeStylus/RealTimeStylusTabletContext.cpp b/Engine/Plugins/Editor/StylusInput/Source/StylusInput/Private/Windows/RealTimeStylus/RealTimeStylusTabletContext.cpp
new file mode 100644
--- /dev/null
+++ b/Engine/Plugins/Editor/StylusInput/Source/StylusInput/Private/Windows/RealTimeStylus/RealTimeStylusTabletContext.cpp
@@ -0,0 +1,38 @@
+﻿// Copyright Epic Games, Inc. All Rights Reserved.
+
+#include "RealTimeStylusTabletContext.h"
+
+#if PLATFORM_WINDOWS
+
+#include <Microsoft/COMPointer.h>
+#include <msinkaut.h>
+
+namespace UE::StylusInput::Private::RealTimeStylus
+{
+	FPacketPropertyConstant PacketPropertyConstants[] = {
+		{EPacketPropertyType::X,                       GUID_PACKETPROPERTY_GUID_X,                        STR_GUID_X                      },
+		{EPacketPropertyType::Y,                       GUID_PACKETPROPERTY_GUID_Y,                        STR_GUID_Y                      },
+		{EPacketPropertyType::Z,                       GUID_PACKETPROPERTY_GUID_Z,                        STR_GUID_Z                      },
+		{EPacketPropertyType::PacketStatus,            GUID_PACKETPROPERTY_GUID_PACKET_STATUS,            STR_GUID_PAKETSTATUS            },
+		{EPacketPropertyType::TimerTick,               GUID_PACKETPROPERTY_GUID_TIMER_TICK,               STR_GUID_TIMERTICK              },
+		{EPacketPropertyType::SerialNumber,            GUID_PACKETPROPERTY_GUID_SERIAL_NUMBER,            STR_GUID_SERIALNUMBER           },
+		{EPacketPropertyType::NormalPressure,          GUID_PACKETPROPERTY_GUID_NORMAL_PRESSURE,          STR_GUID_NORMALPRESSURE         },
+		{EPacketPropertyType::TangentPressure,         GUID_PACKETPROPERTY_GUID_TANGENT_PRESSURE,         STR_GUID_TANGENTPRESSURE        },
+		{EPacketPropertyType::ButtonPressure,          GUID_PACKETPROPERTY_GUID_BUTTON_PRESSURE,          STR_GUID_BUTTONPRESSURE         },
+		{EPacketPropertyType::XTiltOrientation,        GUID_PACKETPROPERTY_GUID_X_TILT_ORIENTATION,       STR_GUID_XTILTORIENTATION       },
+		{EPacketPropertyType::YTiltOrientation,        GUID_PACKETPROPERTY_GUID_Y_TILT_ORIENTATION,       STR_GUID_YTILTORIENTATION       },
+		{EPacketPropertyType::AzimuthOrientation,      GUID_PACKETPROPERTY_GUID_AZIMUTH_ORIENTATION,      STR_GUID_AZIMUTHORIENTATION     },
+		{EPacketPropertyType::AltitudeOrientation,     GUID_PACKETPROPERTY_GUID_ALTITUDE_ORIENTATION,     STR_GUID_ALTITUDEORIENTATION    },
+		{EPacketPropertyType::TwistOrientation,        GUID_PACKETPROPERTY_GUID_TWIST_ORIENTATION,        STR_GUID_TWISTORIENTATION       },
+		{EPacketPropertyType::PitchRotation,           GUID_PACKETPROPERTY_GUID_PITCH_ROTATION,           STR_GUID_PITCHROTATION          },
+		{EPacketPropertyType::RollRotation,            GUID_PACKETPROPERTY_GUID_ROLL_ROTATION,            STR_GUID_ROLLROTATION           },
+		{EPacketPropertyType::YawRotation,             GUID_PACKETPROPERTY_GUID_YAW_ROTATION,             STR_GUID_YAWROTATION            },
+		{EPacketPropertyType::Width,                   GUID_PACKETPROPERTY_GUID_WIDTH,                    STR_GUID_WIDTH                  },
+		{EPacketPropertyType::Height,                  GUID_PACKETPROPERTY_GUID_HEIGHT,                   STR_GUID_HEIGHT                 },
+		{EPacketPropertyType::FingerContactConfidence, GUID_PACKETPROPERTY_GUID_FINGERCONTACTCONFIDENCE,  STR_GUID_FINGERCONTACTCONFIDENCE},
+		{EPacketPropertyType::DeviceContactID,         GUID_PACKETPROPERTY_GUID_DEVICE_CONTACT_ID,        STR_GUID_DEVICE_CONTACT_ID      }
+	};
+	static_assert(std::size(PacketPropertyConstants) == static_cast<uint32>(EPacketPropertyType::Num_Enumerators));
+}
+
+#endif
\ No newline at end of file
diff --git a/Engine/Plugins/Editor/StylusInput/Source/StylusInput/Private/Windows/RealTimeStylus/RealTimeStylusTabletContext.h b/Engine/Plugins/Editor/StylusInput/Source/StylusInput/Private/Windows/RealTimeStylus/RealTimeStylusTabletContext.h
new file mode 100644
--- /dev/null
+++ b/Engine/Plugins/Editor/StylusInput/Source/StylusInput/Private/Windows/RealTimeStylus/RealTimeStylusTabletContext.h
@@ -0,0 +1,169 @@
+﻿// Copyright Epic Games, Inc. All Rights Reserved.
+
+#pragma once
+
+#include <HAL/Platform.h>
+
+#if PLATFORM_WINDOWS
+
+#include <StylusInputPacket.h>
+#include <StylusInputTabletContext.h>
+#include <StylusInputUtils.h>
+#include <Math/IntRect.h>
+#include <Math/Vector.h>
+
+#include "RealTimeStylusUtils.h"
+
+// The following includes are needed for 'GUID' type.
+#include <Windows/AllowWindowsPlatformTypes.h>
+#include <Windows/HideWindowsPlatformTypes.h>
+
+namespace UE::StylusInput::Private::RealTimeStylus
+{
+	enum class EPacketPropertyType : int8
+	{
+		X = 0, // The x-coordinate in the tablet coordinate space. Each packet contains this property by default. The origin (0,0) of the tablet is the upper-left corner.
+		Y, // The y-coordinate in the tablet coordinate space. Each packet contains this property by default. The origin (0,0) of the tablet is the upper-left corner.
+		Z, // The z-coordinate or distance of the pen tip from the tablet surface. The TabletPropertyMetricUnit enumeration type determines the unit of measurement for this property.
+		PacketStatus, // Contains one or more of the following flag values: The cursor is touching the drawing surface (Value = 1). The cursor is inverted. For example, the eraser end of the pen is pointing toward the surface (Value = 2). Not used (Value = 4). The barrel button is pressed (Value = 8).
+		TimerTick, // The time the packet was generated in milliseconds since system start.
+		SerialNumber, // The packet property for identifying the packet. This is the same value you use to retrieve the packet from the packet queue.
+		NormalPressure, // The pressure of the pen tip perpendicular to the tablet surface. The greater the pressure on the pen tip, the more ink that is drawn.
+		TangentPressure, // The pressure of the pen tip along the plane of the tablet surface.
+		ButtonPressure, // The pressure on a pressure sensitive button.
+		XTiltOrientation, // The angle between the y,z-plane and the pen and y-axis plane. Applies to a pen cursor. The value is 0 when the pen is perpendicular to the drawing surface and is positive when the pen is to the right of perpendicular.
+		YTiltOrientation, // The angle between the x,z-plane and the pen and x-axis plane. Applies to a pen cursor. The value is 0 when the pen is perpendicular to the drawing surface and is positive when the pen is upward or away from the user.
+		AzimuthOrientation, // The clockwise rotation of the cursor about the z-axis through a full circular range.
+		AltitudeOrientation, // The angle between the axis of the pen and the surface of the tablet. The value is 0 when the pen is parallel to the surface and 90 when the pen is perpendicular to the surface. The values are negative when the pen is inverted.
+		TwistOrientation, // The clockwise rotation of the cursor about its own axis.
+		PitchRotation, // The packet property that indicates whether the tip is above or below a horizontal line that is perpendicular to the writing surface. Note: This requires a 3D digitizer. The value is positive if the tip is above the line and negative if it is below the line. For example, if you hold the pen in front of you and write on an imaginary wall, the pitch is positive if the tip is above a line extending from you to the wall.
+		RollRotation, // The clockwise rotation of the pen around its own axis. Note: This requires a 3D digitizer.
+		YawRotation, // The angle of the pen to the left or right around the center of its horizontal axis when the pen is horizontal. Note: This requires a 3D digitizer. If you hold the pen in front of you and write on an imaginary wall, zero yaw indicates that the pen is perpendicular to the wall. The value is negative if the tip is to the left of perpendicular and positive if the tip is to the right of perpendicular.
+		Width, // The width of the contact area on a touch digitizer.
+		Height, // The height of the contact area on a touch digitizer.
+		FingerContactConfidence, // The level of confidence that there was finger contact on a touch digitizer.
+		DeviceContactID, // The device contact identifier for a packet.
+
+		Num_Enumerators, // THIS IS NOT A VALID ENUMERATOR, BUT IT CAN BE USED TO QUERY HOW MANY PACKET PROPERTIES THERE ARE.
+		Invalid_Enumerator = INDEX_NONE // THIS IS NOT A VALID ENUMERATOR, BUT IT CAN BE USED TO INDICATE THAT A PACKET PROPERTY ENTRY IS INVALID.
+	};
+
+	enum class ETabletPropertyMetricUnit : int8
+	{
+		Default, // Units are unknown.
+		Inches,
+		Centimeters,
+		Degrees,
+		Radians,
+		Seconds,
+		Pounds,
+		Grams,
+
+		Num_Enumerators, // THIS IS NOT A VALID ENUMERATOR, BUT IT CAN BE USED TO QUERY HOW MANY PACKET PROPERTIES THERE ARE.
+		Invalid_Enumerator = INDEX_NONE // THIS IS NOT A VALID ENUMERATOR, BUT IT CAN BE USED TO INDICATE THAT A PACKET PROPERTY ENTRY IS INVALID.
+	};
+
+	struct FPacketProperty
+	{
+		EPacketPropertyType Type = EPacketPropertyType::Invalid_Enumerator;
+		ETabletPropertyMetricUnit MetricUnit = ETabletPropertyMetricUnit::Default;
+		int32 Minimum = 0;
+		int32 Maximum = 0;
+		float Resolution = 0.0f;
+	};
+
+	struct FPacketPropertyHandler
+	{
+		using FFuncSetProperty = void(FStylusInputPacket&, int32, const int8*);
+		FFuncSetProperty* SetProperty = nullptr;
+
+		static constexpr int32 SetPropertyDataBufferLength = 24;
+		int8 SetPropertyData[SetPropertyDataBufferLength];
+	};
+
+	struct FPacketPropertyConstant
+	{
+		EPacketPropertyType PacketPropertyType;
+		GUID Guid;
+		const wchar_t* StrGuid;
+	};
+
+	extern FPacketPropertyConstant PacketPropertyConstants[static_cast<uint32>(EPacketPropertyType::Num_Enumerators)];
+
+	struct FWindowContext
+	{
+		FVector2d XYScale;
+		FVector2d XYMaximum;
+		FVector2d WindowSize;
+	};
+
+	class FTabletContext final : public IStylusInputTabletContext
+	{
+	public:
+		explicit FTabletContext(const uint32 ID)
+			: ID(ID)
+		{}
+
+		virtual uint32 GetID() const override { return ID; }
+
+		virtual FString GetName() const override
+		{
+			if (!Name.IsEmpty() && !PlugAndPlayID.IsEmpty())
+			{
+				return FString::Format(TEXT("{0} ({1})"), {Name, PlugAndPlayID});
+			}
+			return !Name.IsEmpty() ? Name : PlugAndPlayID;
+		}
+
+		virtual FIntRect GetInputRectangle() const override { return InputRectangle; }
+		virtual ETabletHardwareCapabilities GetHardwareCapabilities() const override { return HardwareCapabilities; }
+		virtual ETabletSupportedProperties GetSupportedProperties() const override { return SupportedProperties; }
+
+		uint32 ID = 0;
+		FString Name;
+		FString PlugAndPlayID;
+		FIntRect InputRectangle;
+		ETabletSupportedProperties SupportedProperties = ETabletSupportedProperties::None;
+		ETabletHardwareCapabilities HardwareCapabilities = ETabletHardwareCapabilities::None;
+		FPacketProperty PacketProperties[static_cast<uint32>(EPacketPropertyType::Num_Enumerators)] = {};
+		FPacketPropertyHandler PacketPropertyHandlers[static_cast<uint32>(EPacketPropertyType::Num_Enumerators)] = {};
+	};
+
+	using FTabletContextContainer = TSharedRefDataContainer<FTabletContext, false>;
+	using FTabletContextThreadSafeContainer = TSharedRefDataContainer<FTabletContext, true>;
+
+	class FStylusButton final : public IStylusInputStylusButton
+	{
+	public:
+		virtual FString GetID() const override { return ID; }
+		virtual FString GetName() const override { return Name; }
+
+		FString ID;
+		FString Name;
+	};
+
+	class FStylusInfo final : public IStylusInputStylusInfo
+	{
+	public:
+		explicit FStylusInfo(const uint32 ID)
+			: ID(ID)
+		{}
+
+		virtual uint32 GetID() const override { return ID; }
+		virtual FString GetName() const override { return Name; }
+		virtual uint32 GetNumButtons() const override { return Buttons.Num(); }
+
+		virtual const IStylusInputStylusButton* GetButton(int32 Index) const override
+		{
+			return 0 <= Index && Index < Buttons.Num() ? reinterpret_cast<const IStylusInputStylusButton*>(&Buttons[Index]) : nullptr;
+		}
+
+		uint32 ID = 0;
+		FString Name;
+		TArray<FStylusButton> Buttons;
+	};
+
+	using FStylusInfoThreadSafeContainer = TSharedRefDataContainer<FStylusInfo, true>;
+}
+
+#endif
diff --git a/Engine/Plugins/Editor/StylusInput/Source/StylusInput/Private/Windows/RealTimeStylus/RealTimeStylusUtils.cpp b/Engine/Plugins/Editor/StylusInput/Source/StylusInput/Private/Windows/RealTimeStylus/RealTimeStylusUtils.cpp
new file mode 100644
--- /dev/null
+++ b/Engine/Plugins/Editor/StylusInput/Source/StylusInput/Private/Windows/RealTimeStylus/RealTimeStylusUtils.cpp
@@ -0,0 +1,28 @@
+﻿// Copyright Epic Games, Inc. All Rights Reserved.
+
+#include "RealTimeStylusUtils.h"
+
+#if PLATFORM_WINDOWS
+
+#include <Windows/AllowWindowsPlatformAtomics.h>
+	#include <comdef.h>
+#include <Windows/HideWindowsPlatformAtomics.h>
+
+#include <StylusInputUtils.h>
+
+namespace UE::StylusInput::Private::RealTimeStylus
+{
+	FString GetCOMErrorMessage(const HRESULT Result)
+	{
+		const _com_error COMError(Result);
+		const LPCTSTR ErrorMessage = COMError.ErrorMessage();
+		return ErrorMessage;
+	}
+
+	void LogCOMError(const FString& Preamble, const HRESULT Result)
+	{
+		LogError(Preamble, FString::Format(TEXT("COM Interface Error: {0}"), {GetCOMErrorMessage(Result)}));
+	}
+}
+
+#endif
diff --git a/Engine/Plugins/Editor/StylusInput/Source/StylusInput/Private/Windows/RealTimeStylus/RealTimeStylusUtils.h b/Engine/Plugins/Editor/StylusInput/Source/StylusInput/Private/Windows/RealTimeStylus/RealTimeStylusUtils.h
new file mode 100644
--- /dev/null
+++ b/Engine/Plugins/Editor/StylusInput/Source/StylusInput/Private/Windows/RealTimeStylus/RealTimeStylusUtils.h
@@ -0,0 +1,38 @@
+﻿// Copyright Epic Games, Inc. All Rights Reserved.
+
+#pragma once
+
+#include <HAL/Platform.h>
+
+#if PLATFORM_WINDOWS
+
+#include <Containers/UnrealString.h>
+#include <Misc/ScopeRWLock.h>
+#include <Templates/SharedPointer.h>
+
+namespace UE::StylusInput::Private::RealTimeStylus
+{
+	void LogCOMError(const FString& Preamble, HRESULT Result);
+
+	inline bool Succeeded(const HRESULT Result, const FString& LogPreamble)
+	{
+		if (Result >= 0)
+		{
+			return true;
+		}
+		LogCOMError(LogPreamble, Result);
+		return false;
+	}
+
+	inline bool Failed(const HRESULT Result, const FString& LogPreamble)
+	{
+		if (Result >= 0)
+		{
+			return false;
+		}
+		LogCOMError(LogPreamble, Result);
+		return true;
+	}
+}
+
+#endif
diff --git a/Engine/Plugins/Editor/StylusInput/Source/StylusInput/Private/Windows/WindowsStylusInputInstance.cpp b/Engine/Plugins/Editor/StylusInput/Source/StylusInput/Private/Windows/WindowsStylusInputInstance.cpp
deleted file mode 100644
diff --git a/Engine/Plugins/Editor/StylusInput/Source/StylusInput/Private/Windows/WindowsStylusInputInstance.h b/Engine/Plugins/Editor/StylusInput/Source/StylusInput/Private/Windows/WindowsStylusInputInstance.h
deleted file mode 100644
diff --git a/Engine/Plugins/Editor/StylusInput/Source/StylusInput/Private/Windows/WindowsStylusInputInterface.cpp b/Engine/Plugins/Editor/StylusInput/Source/StylusInput/Private/Windows/WindowsStylusInputInterface.cpp
deleted file mode 100644
diff --git a/Engine/Plugins/Editor/StylusInput/Source/StylusInput/Private/Windows/WindowsStylusInputInterface.h b/Engine/Plugins/Editor/StylusInput/Source/StylusInput/Private/Windows/WindowsStylusInputInterface.h
deleted file mode 100644
diff --git a/Engine/Plugins/Editor/StylusInput/Source/StylusInput/Private/Windows/WindowsStylusInputPlatformAPI.cpp b/Engine/Plugins/Editor/StylusInput/Source/StylusInput/Private/Windows/WindowsStylusInputPlatformAPI.cpp
deleted file mode 100644
diff --git a/Engine/Plugins/Editor/StylusInput/Source/StylusInput/Private/Windows/WindowsStylusInputPlatformAPI.h b/Engine/Plugins/Editor/StylusInput/Source/StylusInput/Private/Windows/WindowsStylusInputPlatformAPI.h
deleted file mode 100644
diff --git a/Engine/Plugins/Editor/StylusInput/Source/StylusInput/Private/Windows/WindowsStylusInputPluginAsync.cpp b/Engine/Plugins/Editor/StylusInput/Source/StylusInput/Private/Windows/WindowsStylusInputPluginAsync.cpp
deleted file mode 100644
diff --git a/Engine/Plugins/Editor/StylusInput/Source/StylusInput/Private/Windows/WindowsStylusInputPluginAsync.h b/Engine/Plugins/Editor/StylusInput/Source/StylusInput/Private/Windows/WindowsStylusInputPluginAsync.h
deleted file mode 100644
diff --git a/Engine/Plugins/Editor/StylusInput/Source/StylusInput/Private/Windows/WindowsStylusInputPluginBase.cpp b/Engine/Plugins/Editor/StylusInput/Source/StylusInput/Private/Windows/WindowsStylusInputPluginBase.cpp
deleted file mode 100644
diff --git a/Engine/Plugins/Editor/StylusInput/Source/StylusInput/Private/Windows/WindowsStylusInputPluginBase.h b/Engine/Plugins/Editor/StylusInput/Source/StylusInput/Private/Windows/WindowsStylusInputPluginBase.h
deleted file mode 100644
diff --git a/Engine/Plugins/Editor/StylusInput/Source/StylusInput/Private/Windows/WindowsStylusInputPluginSync.cpp b/Engine/Plugins/Editor/StylusInput/Source/StylusInput/Private/Windows/WindowsStylusInputPluginSync.cpp
deleted file mode 100644
diff --git a/Engine/Plugins/Editor/StylusInput/Source/StylusInput/Private/Windows/WindowsStylusInputPluginSync.h b/Engine/Plugins/Editor/StylusInput/Source/StylusInput/Private/Windows/WindowsStylusInputPluginSync.h
deleted file mode 100644
diff --git a/Engine/Plugins/Editor/StylusInput/Source/StylusInput/Private/Windows/WindowsStylusInputStats.h b/Engine/Plugins/Editor/StylusInput/Source/StylusInput/Private/Windows/WindowsStylusInputStats.h
deleted file mode 100644
diff --git a/Engine/Plugins/Editor/StylusInput/Source/StylusInput/Private/Windows/WindowsStylusInputTabletContext.cpp b/Engine/Plugins/Editor/StylusInput/Source/StylusInput/Private/Windows/WindowsStylusInputTabletContext.cpp
deleted file mode 100644
diff --git a/Engine/Plugins/Editor/StylusInput/Source/StylusInput/Private/Windows/WindowsStylusInputTabletContext.h b/Engine/Plugins/Editor/StylusInput/Source/StylusInput/Private/Windows/WindowsStylusInputTabletContext.h
deleted file mode 100644
diff --git a/Engine/Plugins/Editor/StylusInput/Source/StylusInput/Private/Windows/WindowsStylusInputUtils.cpp b/Engine/Plugins/Editor/StylusInput/Source/StylusInput/Private/Windows/WindowsStylusInputUtils.cpp
deleted file mode 100644
diff --git a/Engine/Plugins/Editor/StylusInput/Source/StylusInput/Private/Windows/WindowsStylusInputUtils.h b/Engine/Plugins/Editor/StylusInput/Source/StylusInput/Private/Windows/WindowsStylusInputUtils.h
deleted file mode 100644
