StylusInput: Add Wintab support (needs to be enabled explicitly in the StylusInput.Build.cs file)

diff --git a/Engine/Plugins/Editor/StylusInput/Source/StylusInput/Private/StylusInput.cpp b/Engine/Plugins/Editor/StylusInput/Source/StylusInput/Private/StylusInput.cpp
--- a/Engine/Plugins/Editor/StylusInput/Source/StylusInput/Private/StylusInput.cpp
+++ b/Engine/Plugins/Editor/StylusInput/Source/StylusInput/Private/StylusInput.cpp
@@ -6,7 +6,10 @@
 #include "StylusInputUtils.h"
 
 #if PLATFORM_WINDOWS
-	#include <Windows/RealTimeStylus/RealTimeStylusInterface.h>
+	#if STYLUSINPUT_USE_WINTAB
+		#include "Windows/Wintab/WintabInterface.h"
+	#endif
+	#include "Windows/RealTimeStylus/RealTimeStylusInterface.h"
 #endif
 
 namespace UE::StylusInput
@@ -15,10 +18,10 @@ namespace UE::StylusInput
 	{
 		class FStylusInputImpl;
 
-		class FNotSupportedStylusInputInterface final : public IStylusInputInterface
+		class FNotSupportedInterface final : public IStylusInputInterface
 		{
 		public:
-			virtual ~FNotSupportedStylusInputInterface() override = default;
+			virtual ~FNotSupportedInterface() override = default;
 
 			virtual IStylusInputInstance* CreateInstance(SWindow& Window) override { return nullptr; }
 
@@ -37,7 +40,7 @@ namespace UE::StylusInput
 		{
 		public:
 			FStylusInputImpl()
-				: Interface(FPlatformInterface::Create())
+				: Interface(SelectInterface())
 			{
 			}
 
@@ -58,13 +61,30 @@ namespace UE::StylusInput
 			}
 
 		private:
-			using FPlatformInterface =
+			static TUniquePtr<IStylusInputInterface> SelectInterface()
+			{
 #if PLATFORM_WINDOWS
-				RealTimeStylus::FRealTimeStylusInterface
+	#if STYLUSINPUT_USE_WINTAB
+
+				// Prefer using Wintab, and fall back to RealTimeStylus if not available.
+				TUniquePtr<IStylusInputInterface> Interface = Wintab::FWintabInterface::Create();
+				if (!Interface)
+				{
+					Interface = RealTimeStylus::FRealTimeStylusInterface::Create();
+				}
+				return Interface;
+
+	#else
+
+				return RealTimeStylus::FRealTimeStylusInterface::Create();
+
+	#endif
 #else
-				FNotSupportedStylusInputInterface
+
+				return FNotSupportedInterface::Create();
+
 #endif
-				;
+			}
 
 			TUniquePtr<IStylusInputInterface> Interface;
 		};
diff --git a/Engine/Plugins/Editor/StylusInput/Source/StylusInput/Private/Windows/Wintab/WintabAPI.cpp b/Engine/Plugins/Editor/StylusInput/Source/StylusInput/Private/Windows/Wintab/WintabAPI.cpp
new file mode 100644
--- /dev/null
+++ b/Engine/Plugins/Editor/StylusInput/Source/StylusInput/Private/Windows/Wintab/WintabAPI.cpp
@@ -0,0 +1,162 @@
+﻿// Copyright Epic Games, Inc. All Rights Reserved.
+
+#include "WintabAPI.h"
+
+#if PLATFORM_WINDOWS && STYLUSINPUT_USE_WINTAB
+
+#include <StylusInputUtils.h>
+
+#define LOG_PREAMBLE "WintabAPI"
+
+namespace UE::StylusInput::Private::Wintab
+{
+	FWintabAPI::FWintabAPI()
+	{
+		auto GetDllHandle = [&DllHandles = DllHandles](const TCHAR* DllName) -> void* {
+			void* DllHandle = FWindowsPlatformProcess::GetDllHandle(DllName);
+			if (DllHandle)
+			{
+				DllHandles.Add(DllHandle);
+				LogVerbose(LOG_PREAMBLE, FString::Format(TEXT("Retrieved DLL handle for {0}."), {DllName}));
+			}
+			else
+			{
+				LogError(LOG_PREAMBLE, FString::Format(TEXT("Could not get DLL handle for {0}."), {DllName}));
+			}
+			return DllHandle;
+		};
+
+		auto GetDllExport = []<typename FFuncType>(const TCHAR* DllName, void* DllHandle, const TCHAR* ExportName, FFuncType& ExportHandle)
+		{
+			void* FuncPtr = FWindowsPlatformProcess::GetDllExport(DllHandle, ExportName);
+			if (FuncPtr)
+			{
+				LogVerbose(LOG_PREAMBLE, FString::Format(TEXT("Retrieved DLL export '{0}' from {1}."), {ExportName, DllName}));
+			}
+			else
+			{
+				LogError(LOG_PREAMBLE, FString::Format(TEXT("Could not get DLL export '{0}' in {1}."), {ExportName, DllName}));
+			}
+			ExportHandle = static_cast<FFuncType>(FuncPtr);
+		};
+
+		const TCHAR* User32DllName = TEXT("User32.dll");
+		void* User32DllHandle = GetDllHandle(User32DllName);
+		GetDllExport(User32DllName, User32DllHandle, TEXT("GetWindowRect"), GetWindowRect);
+		GetDllExport(User32DllName, User32DllHandle, TEXT("GetSystemMetrics"), GetSystemMetrics);
+
+		const TCHAR* Wintab32DllName = TEXT("Wintab32.dll");
+		void* Wintab32DllHandle = GetDllHandle(Wintab32DllName);
+		GetDllExport(Wintab32DllName, Wintab32DllHandle, TEXT("WTInfoW"), WtInfo);
+		GetDllExport(Wintab32DllName, Wintab32DllHandle, TEXT("WTOpenW"), WtOpen);
+		GetDllExport(Wintab32DllName, Wintab32DllHandle, TEXT("WTClose"), WtClose);
+		GetDllExport(Wintab32DllName, Wintab32DllHandle, TEXT("WTPacket"), WtPacket);
+
+		if (WtInfo)
+		{
+			bWintabIsAvailable = WtInfo(0, 0, nullptr) != 0;
+			if (!bWintabIsAvailable)
+			{
+				LogVerbose(LOG_PREAMBLE, TEXT("Wintab Digitizer Services are not available."));
+			}
+			else
+			{
+				constexpr UINT OutputBufferSize = 1024;
+				TCHAR OutputBuffer[OutputBufferSize];
+				if (ensure(WtInfo(WTI_INTERFACE, IFC_WINTABID, OutputBuffer) <= OutputBufferSize))
+				{
+					InterfaceInfo.ID = OutputBuffer;
+
+					LogVerbose(LOG_PREAMBLE, FString::Format(TEXT("Wintab ID: {0}"), {OutputBuffer}));
+				}
+
+				WORD SpecVersion;
+				if (ensure(WtInfo(WTI_INTERFACE, IFC_SPECVERSION, &SpecVersion) == sizeof(SpecVersion)))
+				{
+					InterfaceInfo.SpecificationVersion = {static_cast<uint16>(SpecVersion >> 8), static_cast<uint16>(SpecVersion & 0xFF)};
+
+					LogVerbose(LOG_PREAMBLE, FString::Format(TEXT("Wintab Specification Version: {0}.{1}"),
+					                                         {InterfaceInfo.SpecificationVersion.Major, InterfaceInfo.SpecificationVersion.Minor}));
+				}
+
+				WORD ImplVersion;
+				if (ensure(WtInfo(WTI_INTERFACE, IFC_IMPLVERSION, &ImplVersion) == sizeof(ImplVersion)))
+				{
+					InterfaceInfo.ImplementationVersion = {static_cast<uint16>(ImplVersion >> 8), static_cast<uint16>(ImplVersion & 0xFF)};
+
+					LogVerbose(LOG_PREAMBLE, FString::Format(TEXT("Wintab Implementation Version: {0}.{1}"),
+					                                         {InterfaceInfo.ImplementationVersion.Major, InterfaceInfo.ImplementationVersion.Minor}));
+				}
+
+				/*UINT NumDevices;
+				if (ensure(WtInfo(WTI_INTERFACE, IFC_NDEVICES, &NumDevices) == sizeof(NumDevices)))
+				{
+					LogVerbose(LOG_PREAMBLE, FString::Format(TEXT("Wintab Number of Connected Devices: {0}"), {NumDevices}));
+				}*/
+
+				UINT MaxNumCursorTypes;
+				if (ensure(WtInfo(WTI_INTERFACE, IFC_NCURSORS, &MaxNumCursorTypes) == sizeof(MaxNumCursorTypes)))
+				{
+					InterfaceInfo.MaximumNumberOfCursorTypes = MaxNumCursorTypes;
+
+					LogVerbose(LOG_PREAMBLE, FString::Format(TEXT("Wintab Maximum Number of Cursor Types: {0}"), {MaxNumCursorTypes}));
+				}
+
+				UINT NumContexts;
+				if (ensure(WtInfo(WTI_INTERFACE, IFC_NCONTEXTS, &NumContexts) == sizeof(NumContexts)))
+				{
+					InterfaceInfo.MaximumNumberOfContexts = NumContexts;
+
+					LogVerbose(LOG_PREAMBLE, FString::Format(TEXT("Wintab Maximum Number of Contexts: {0}"), {NumContexts}));
+				}
+
+				UINT ContextOptions;
+				if (ensure(WtInfo(WTI_INTERFACE, IFC_CTXOPTIONS, &ContextOptions) == sizeof(ContextOptions)))
+				{
+					InterfaceInfo.bSupportsPacketMessages = ContextOptions & CXO_MESSAGES;
+					InterfaceInfo.bSupportsCursorChangeMessages = ContextOptions & CXO_CSRMESSAGES;
+
+					LogVerbose(LOG_PREAMBLE, FString::Format(TEXT("Wintab Supports System Cursor Contexts: {0}"), {ContextOptions & CXO_SYSTEM ? TEXT("Yes") : TEXT("No")}));
+					LogVerbose(LOG_PREAMBLE, FString::Format(TEXT("Wintab Supports Pen Windows Contexts: {0}"), {ContextOptions & CXO_PEN ? TEXT("Yes") : TEXT("No")}));
+					LogVerbose(LOG_PREAMBLE, FString::Format(TEXT("Wintab Supports WT_PACKET Messages: {0}"), {ContextOptions & CXO_MESSAGES ? TEXT("Yes") : TEXT("No")}));
+					LogVerbose(LOG_PREAMBLE, FString::Format(TEXT("Wintab Supports WT_CSRCHANGE  Messages: {0}"), {ContextOptions & CXO_CSRMESSAGES ? TEXT("Yes") : TEXT("No")}));
+					LogVerbose(LOG_PREAMBLE, FString::Format(TEXT("Wintab Supports Margins: {0}"), {ContextOptions & CXO_MARGIN ? TEXT("Yes") : TEXT("No")}));
+					LogVerbose(LOG_PREAMBLE, FString::Format(TEXT("Wintab Supports Inside Margins: {0}"), {ContextOptions & CXO_MGNINSIDE ? TEXT("Yes") : TEXT("No")}));
+				}
+			}
+		}
+	}
+
+	FWintabAPI::~FWintabAPI()
+	{
+		for (void* DllHandle : DllHandles)
+		{
+			FWindowsPlatformProcess::FreeDllHandle(DllHandle);
+		}
+	}
+
+	const FWintabAPI& FWintabAPI::GetInstance()
+	{
+		static FWintabAPI WintabAPI;
+		return WintabAPI;
+	}
+
+	bool FWintabAPI::SatisfiesRequirements() const
+	{
+		bool bSuccess = true;
+
+		bSuccess &= bWintabIsAvailable;
+		bSuccess &= GetWindowRect != nullptr;
+		bSuccess &= GetSystemMetrics != nullptr;
+		bSuccess &= WtInfo != nullptr;
+		bSuccess &= WtOpen != nullptr;
+		bSuccess &= WtClose != nullptr;
+		bSuccess &= WtPacket != nullptr;
+
+		return bSuccess;
+	}
+}
+
+#undef LOG_PREAMBLE
+
+#endif
diff --git a/Engine/Plugins/Editor/StylusInput/Source/StylusInput/Private/Windows/Wintab/WintabAPI.h b/Engine/Plugins/Editor/StylusInput/Source/StylusInput/Private/Windows/Wintab/WintabAPI.h
new file mode 100644
--- /dev/null
+++ b/Engine/Plugins/Editor/StylusInput/Source/StylusInput/Private/Windows/Wintab/WintabAPI.h
@@ -0,0 +1,171 @@
+﻿// Copyright Epic Games, Inc. All Rights Reserved.
+
+#pragma once
+
+#include <HAL/Platform.h>
+
+#if PLATFORM_WINDOWS && STYLUSINPUT_USE_WINTAB
+
+#include <Containers/Array.h>
+#include <Containers/UnrealString.h>
+#include <Templates/UniquePtr.h>
+#include <Windows/WindowsPlatformProcess.h>
+
+#define WIN32_LEAN_AND_MEAN
+#include <windows.h>
+#include <WINTAB.H>
+#define PACKETDATA (/*PK_CONTEXT |*/ PK_STATUS | PK_TIME | /*PK_CHANGED |*/ PK_SERIAL_NUMBER | PK_CURSOR | PK_BUTTONS | PK_X | PK_Y | PK_Z | PK_NORMAL_PRESSURE | PK_TANGENT_PRESSURE | PK_ORIENTATION | PK_ROTATION)
+#define PACKETMODE PK_BUTTONS
+#include <PKTDEF.H>
+
+#define WTINFO_WORKAROUND_FOR_OUTPUT_BUFFER_SIZE_BUG 0
+
+namespace UE::StylusInput::Private::Wintab
+{
+	struct FWintabInterfaceInfo
+	{
+		struct FVersion
+		{
+			uint16 Major;
+			uint16 Minor;
+		};
+
+		FString ID;
+		FVersion SpecificationVersion = { 0, 0 };
+		FVersion ImplementationVersion = { 0, 0 };
+		uint32 MaximumNumberOfCursorTypes = 0;
+		uint32 MaximumNumberOfContexts = 0;
+		bool bSupportsPacketMessages = false;
+		bool bSupportsCursorChangeMessages = false;
+	};
+
+	class FWintabAPI
+	{
+	public:
+		static const FWintabAPI& GetInstance();
+
+		bool SatisfiesRequirements() const;
+
+		using FFuncGetWindowRect = BOOL(HWND, LPRECT);
+		FFuncGetWindowRect* GetWindowRect = nullptr;
+
+		using FFuncGetSystemMetrics = int(int);
+		FFuncGetSystemMetrics* GetSystemMetrics = nullptr;
+
+		using FFuncWtInfo = UINT(UINT, UINT, LPVOID);
+		FFuncWtInfo* WtInfo = nullptr;
+
+		using FFuncWtOpen = HCTX(HWND, LPLOGCONTEXT, BOOL);
+		FFuncWtOpen* WtOpen = nullptr;
+
+		using FFuncWtClose = BOOL(HCTX);
+		FFuncWtClose* WtClose = nullptr;
+
+		using FFuncWtPacket = BOOL(HCTX, UINT, LPVOID);
+		FFuncWtPacket* WtPacket = nullptr;
+
+		FWintabInterfaceInfo InterfaceInfo;
+
+	private:
+
+		FWintabAPI();
+		~FWintabAPI();
+
+		TArray<void*> DllHandles;
+
+		bool bWintabIsAvailable = false;
+	};
+
+	inline uint32 HctxToUint32(const HCTX Hctx)
+	{
+		return IntCastChecked<uint32>(reinterpret_cast<uint64>(Hctx));
+	}
+
+	inline float Fix32ToFloat(const FIX32 Value)
+	{
+		return static_cast<float>(HIWORD(Value)) + LOWORD(Value) / 65536.0f;
+	}
+
+#if WTINFO_WORKAROUND_FOR_OUTPUT_BUFFER_SIZE_BUG
+
+	/*
+	 * Apparently, some drivers have issues with a nullptr being passed as output buffer for WtInfo() calls.
+	 * The workaround is to just always have a fixed size buffer and hope for the best.
+	 */
+
+	class FWintabInfoOutputBuffer
+	{
+	public:
+		void Allocate(const UINT Category, const UINT Index)
+		{
+			if (!Buffer.IsValid())
+			{
+				Buffer.Reset(new TCHAR[Size]);
+			}
+			Buffer[0] = 0;
+		}
+
+		TCHAR* Get() const { return Buffer.Get(); }
+		constexpr UINT SizeInBytes() { return Size * sizeof(TCHAR); }
+
+	private:
+		enum { Size = 1024 };
+		TUniquePtr<TCHAR[]> Buffer;
+	};
+
+#else
+
+	class FWintabInfoOutputBuffer
+	{
+	public:
+		FWintabInfoOutputBuffer()
+			: WintabAPI(FWintabAPI::GetInstance())
+		{
+		}
+
+		void Allocate(const UINT Category, const UINT Index)
+		{
+			// If the following call crashes, consider enabling the workaround available via WTINFO_WORKAROUND_FOR_OUTPUT_BUFFER_SIZE_BUG.
+			const UINT RequiredSize = WintabAPI.WtInfo(Category, Index, nullptr)  / sizeof(TCHAR);
+
+			if (Size < RequiredSize)
+			{
+				Size = RequiredSize;
+				Buffer.Reset(new TCHAR[Size]);
+			}
+			if (Size > 0)
+			{
+				Buffer[0] = 0;
+			}
+		}
+
+		TCHAR* Get() const { return Buffer.Get(); }
+		UINT SizeInBytes() const { return Size * sizeof(TCHAR); }
+
+	private:
+		const FWintabAPI& WintabAPI;
+		TUniquePtr<TCHAR[]> Buffer;
+		UINT Size = 0;
+	};
+
+#endif
+}
+
+/*
+ * TODO 
+ * [ ] RealTimeStylus replace TimerTick value at earliest opportunity
+ * [ ] DPI scaling?
+ * [ ] Handle removing/adding tablet while running (evaluate WT_INFOCHANGE event)
+ * [ ] PK_BUTTONS (plus RealtimeStylus buttons)
+ * [ ] Queue size
+ * [ ] Query multiple packets at once
+ * [ ] Async
+ * [ ] Deal with cursors not having a physical ID
+ * [ ] Make available by default (with additional editor settings)
+ * [ ] Propagate buttons data for both Wintab and Windows (StylusButtonDown/StylusButtonDown events)
+ * [ ] Allegedly, some Wintab libraries don't handle relative button presses, so maybe switch to absolute?
+ * [ ] Lost packet stats?
+ * [ ] RealTimeStylus StylusInRange and StylusOutOfRange events (re-evaluate stylus info?)
+ */
+
+#endif
diff --git a/Engine/Plugins/Editor/StylusInput/Source/StylusInput/Private/Windows/Wintab/WintabInstance.cpp b/Engine/Plugins/Editor/StylusInput/Source/StylusInput/Private/Windows/Wintab/WintabInstance.cpp
new file mode 100644
--- /dev/null
+++ b/Engine/Plugins/Editor/StylusInput/Source/StylusInput/Private/Windows/Wintab/WintabInstance.cpp
@@ -0,0 +1,839 @@
+﻿// Copyright Epic Games, Inc. All Rights Reserved.
+
+#include "WintabInstance.h"
+
+#if PLATFORM_WINDOWS && STYLUSINPUT_USE_WINTAB
+
+#include <StylusInput.h>
+#include <StylusInputUtils.h>
+#include <Windows/WindowsApplication.h>
+#include <Framework/Application/SlateApplication.h>
+
+#include "WintabAPI.h"
+
+#define LOCTEXT_NAMESPACE "WintabInstance"
+#define LOG_PREAMBLE "WintabInstance"
+
+namespace UE::StylusInput::Private::Wintab
+{
+	bool SetupWintabContext(const FWintabAPI& WintabAPI, UINT WintabContextIndex, LOGCONTEXT& WintabContext)
+	{
+		bool bSuccess = true;
+
+		WintabContext.lcPktData = PACKETDATA;
+		WintabContext.lcPktMode = PACKETMODE;
+		WintabContext.lcOptions |= CXO_CSRMESSAGES | CXO_MESSAGES | CXO_SYSTEM;
+
+		AXIS AxisX = {};
+		AXIS AxisY = {};
+		AXIS AxisZ = {};
+
+		bSuccess &= WintabAPI.WtInfo(WTI_DEVICES + WintabContextIndex, DVC_X, &AxisX) == sizeof(AxisX);
+		bSuccess &= WintabAPI.WtInfo(WTI_DEVICES + WintabContextIndex, DVC_Y, &AxisY) == sizeof(AxisY);
+		bSuccess &= WintabAPI.WtInfo(WTI_DEVICES + WintabContextIndex, DVC_Z, &AxisZ) == sizeof(AxisZ);
+
+		if (!bSuccess)
+		{
+			LogError(LOG_PREAMBLE, FString::Format(TEXT("Failed to query device properties for Wintab context with index {0}."), {WintabContextIndex}));
+			return false;
+		}
+
+		WintabContext.lcOutExtX = AxisX.axMax - AxisX.axMin + 1;
+		WintabContext.lcOutExtY = AxisY.axMax - AxisY.axMin + 1;
+		WintabContext.lcOutExtZ = AxisZ.axMax - AxisZ.axMin + 1;
+
+		// In Wintab, the tablet origin is lower left. Move origin to upper left so that it corresponds to screen origin.
+		WintabContext.lcOutExtY = -WintabContext.lcOutExtY;
+
+		return bSuccess;
+	}
+
+	bool SetupTabletContextMetadata(const FWintabAPI& WintabAPI, UINT WintabContextIndex, HCTX WintabContextHandle,
+	                                const LOGCONTEXT& WintabContext, FTabletContext& TabletContext)
+	{
+		bool bSuccess = true;
+
+		TabletContext.WintabContextIndex = WintabContextIndex;
+		TabletContext.WintabContextHandle = WintabContextHandle;
+
+		TabletContext.InputRectangle = {
+			WintabContext.lcOutOrgX, WintabContext.lcOutOrgY,
+			WintabContext.lcOutOrgX + WintabContext.lcOutExtX, WintabContext.lcOutOrgY - WintabContext.lcOutExtY
+		};
+
+		FWintabInfoOutputBuffer OutputBuffer;
+
+		TabletContext.Name = WintabContext.lcName;
+		if (TabletContext.Name.IsEmpty())
+		{
+			OutputBuffer.Allocate(WTI_DEVICES + WintabContextIndex, DVC_NAME);
+			if (WintabAPI.WtInfo(WTI_DEVICES + WintabContextIndex, DVC_NAME, OutputBuffer.Get()) <= OutputBuffer.SizeInBytes())
+			{
+				TabletContext.Name = OutputBuffer.Get();
+			}
+			else
+			{
+				LogWarning(LOG_PREAMBLE, FString::Format(TEXT("Failed to query device name for Wintab context with index {0}."), {WintabContextIndex}));
+			}
+		}
+
+		OutputBuffer.Allocate(WTI_DEVICES + WintabContextIndex, DVC_PNPID);
+		if (WintabAPI.WtInfo(WTI_DEVICES + WintabContextIndex, DVC_PNPID, OutputBuffer.Get()) <= OutputBuffer.SizeInBytes())
+		{
+			TabletContext.PlugAndPlayID = OutputBuffer.Get();
+		}
+		else
+		{
+			TabletContext.PlugAndPlayID.Empty();
+			LogWarning(LOG_PREAMBLE, FString::Format(TEXT("Failed to query Plug and Play ID for Wintab context with index {0}."), {WintabContextIndex}));
+		}
+
+		UINT HardwareCapabilities;
+		if (WintabAPI.WtInfo(WTI_DEVICES + WintabContextIndex, DVC_HARDWARE, &HardwareCapabilities) != sizeof(HardwareCapabilities))
+		{
+			LogError(LOG_PREAMBLE, FString::Format(TEXT("Failed to query hardware capabilities for Wintab context with index {0}."), {WintabContextIndex}));
+			return false;
+		}
+
+		TabletContext.HardwareCapabilities =
+			(HardwareCapabilities & HWC_INTEGRATED ? ETabletHardwareCapabilities::Integrated : ETabletHardwareCapabilities::None)
+			| (HardwareCapabilities & HWC_TOUCH ? ETabletHardwareCapabilities::CursorMustTouch : ETabletHardwareCapabilities::None)
+			| (HardwareCapabilities & HWC_HARDPROX ? ETabletHardwareCapabilities::HardProximity : ETabletHardwareCapabilities::None)
+			| (HardwareCapabilities & HWC_PHYSID_CURSORS ? ETabletHardwareCapabilities::CursorsHavePhysicalIds : ETabletHardwareCapabilities::None);
+
+		/* CURSORS */
+
+		UINT FirstCursor, NumCursors;
+		if (WintabAPI.WtInfo(WTI_DEVICES + WintabContextIndex, DVC_FIRSTCSR, &FirstCursor) == sizeof(FirstCursor)
+			&& WintabAPI.WtInfo(WTI_DEVICES + WintabContextIndex, DVC_NCSRTYPES, &NumCursors) == sizeof(NumCursors))
+		{
+			TabletContext.WintabFirstCursor = FirstCursor;
+			TabletContext.WintabNumCursors = NumCursors;
+		}
+		else
+		{
+			LogError(LOG_PREAMBLE, FString::Format(TEXT("Failed to query cursor range for Wintab context with index {0}."), {WintabContextIndex}));
+			return false;
+		}
+
+		return bSuccess;
+	}
+
+	void SetPropertyStatus(FStylusInputPacket& Packet, const PACKET& WintabPacket, const int8* Data)
+	{
+		// Todo incorporate if tablet context is actually able to support Z values (or proximity?); if not, probably any packet is touching the tablet 
+		Packet.PenStatus = (WintabPacket.pkStatus & TPS_INVERT ? EPenStatus::CursorIsInverted : EPenStatus::None)
+			| (WintabPacket.pkZ <= 0 ? EPenStatus::CursorIsTouching : EPenStatus::None);
+	}
+	static_assert(std::is_same_v<FPacketPropertyHandler::FFuncSetProperty, decltype(SetPropertyStatus)>);
+
+	void SetPropertyTime(FStylusInputPacket& Packet, const PACKET& WintabPacket, const int8* Data)
+	{
+		Packet.TimerTick = WintabPacket.pkTime;
+	}
+	static_assert(std::is_same_v<FPacketPropertyHandler::FFuncSetProperty, decltype(SetPropertyTime)>);
+
+	void SetPropertySerialNumber(FStylusInputPacket& Packet, const PACKET& WintabPacket, const int8* Data)
+	{
+		Packet.SerialNumber = WintabPacket.pkSerialNumber;
+	}
+	static_assert(std::is_same_v<FPacketPropertyHandler::FFuncSetProperty, decltype(SetPropertySerialNumber)>);
+
+	void SetPropertyCursor(FStylusInputPacket& Packet, const PACKET& WintabPacket, const int8* Data)
+	{
+		const FWintabMessageHandler* MessageHandler = *reinterpret_cast<FWintabMessageHandler *const *>(Data);
+		Packet.CursorID = MessageHandler->GetCurrentStylusID();
+	}
+	static_assert(std::is_same_v<FPacketPropertyHandler::FFuncSetProperty, decltype(SetPropertyCursor)>);
+
+	void AssignSetPropertyCursorData(int8 *const Data, const FWintabMessageHandler* MessageHandler)
+	{
+		*reinterpret_cast<const FWintabMessageHandler**>(Data) = MessageHandler;
+	}
+
+	void SetPropertyButtons(FStylusInputPacket& Packet, const PACKET& WintabPacket, const int8* Data)
+	{
+		// Todo 
+		/*const WORD ButtonNumber = LOWORD(WintabPacket.pkButtons);
+		const WORD ButtonState = HIWORD(WintabPacket.pkButtons);*/
+	}
+	static_assert(std::is_same_v<FPacketPropertyHandler::FFuncSetProperty, decltype(SetPropertyButtons)>);
+
+	void SetPropertyX(FStylusInputPacket& Packet, const PACKET& WintabPacket, const int8* Data)
+	{
+		const int32& TabletContextInputWidth = reinterpret_cast<const int*>(Data)[0];
+		const int32& VirtualScreenWidth = reinterpret_cast<const int*>(Data)[1];
+
+		const RECT *const * WindowRect = reinterpret_cast<const RECT *const *>(Data + sizeof(int32) * 2);
+		const int32& WindowOffsetX = (*WindowRect)->left;
+
+		Packet.X = static_cast<float>(WintabPacket.pkX) * VirtualScreenWidth / TabletContextInputWidth - WindowOffsetX;
+	}
+	static_assert(std::is_same_v<FPacketPropertyHandler::FFuncSetProperty, decltype(SetPropertyX)>);
+
+	void AssignSetPropertyXData(int8 *const Data, const int32 TabletContextInputWidth, const int32 VirtualScreenWidth, const RECT* WindowRect)
+	{
+		int32* DataInt = reinterpret_cast<int32*>(Data);
+		*DataInt++ = TabletContextInputWidth;
+		*DataInt++ = VirtualScreenWidth;
+
+		const RECT** DataRect = reinterpret_cast<const RECT**>(DataInt);
+		*DataRect = WindowRect;
+	}
+
+	void SetPropertyY(FStylusInputPacket& Packet, const PACKET& WintabPacket, const int8* Data)
+	{
+		const int32& TabletContextInputHeight = reinterpret_cast<const int*>(Data)[0];
+		const int32& VirtualScreenHeight = reinterpret_cast<const int*>(Data)[1];
+
+		const RECT *const * WindowRect = reinterpret_cast<const RECT *const *>(Data + sizeof(int32) * 2);
+		const int32& WindowOffsetY = (*WindowRect)->top;
+
+		Packet.Y = static_cast<float>(WintabPacket.pkY) * VirtualScreenHeight / TabletContextInputHeight - WindowOffsetY;
+	}
+	static_assert(std::is_same_v<FPacketPropertyHandler::FFuncSetProperty, decltype(SetPropertyY)>);
+
+	void AssignSetPropertyYData(int8 *const Data, const int32 TabletContextInputHeight, const int32 VirtualScreenHeight, const RECT* WindowRect)
+	{
+		int32* DataInt = reinterpret_cast<int32*>(Data);
+		*DataInt++ = TabletContextInputHeight;
+		*DataInt++ = VirtualScreenHeight;
+
+		const RECT** DataRect = reinterpret_cast<const RECT**>(DataInt);
+		*DataRect = WindowRect;
+	}
+
+	void SetPropertyZ(FStylusInputPacket& Packet, const PACKET& WintabPacket, const int8* Data)
+	{
+		const float& InvResolution = *reinterpret_cast<const float*>(Data);
+
+		Packet.Z = static_cast<float>(WintabPacket.pkZ) * InvResolution;
+	}
+	static_assert(std::is_same_v<FPacketPropertyHandler::FFuncSetProperty, decltype(SetPropertyZ)>);
+
+	void AssignSetPropertyZData(int8 *const Data, const float Resolution)
+	{
+		*reinterpret_cast<float*>(Data) = 1.0f / Resolution;
+	}
+
+	void SetPropertyNormalPressure(FStylusInputPacket& Packet, const PACKET& WintabPacket, const int8* Data)
+	{
+		const float& InvExtent = *reinterpret_cast<const float*>(Data);
+
+		Packet.NormalPressure = static_cast<float>(WintabPacket.pkNormalPressure) * InvExtent;
+	}
+	static_assert(std::is_same_v<FPacketPropertyHandler::FFuncSetProperty, decltype(SetPropertyNormalPressure)>);
+
+	void AssignSetPropertyNormalPressureData(int8 *const Data, const int32 Extent)
+	{
+		*reinterpret_cast<float*>(Data) = 1.0f / Extent;
+	}
+
+	void SetPropertyTangentPressure(FStylusInputPacket& Packet, const PACKET& WintabPacket, const int8* Data)
+	{
+		const float& InvExtent = *reinterpret_cast<const float*>(Data);
+
+		Packet.TangentPressure = static_cast<float>(WintabPacket.pkTangentPressure) * InvExtent;
+	}
+	static_assert(std::is_same_v<FPacketPropertyHandler::FFuncSetProperty, decltype(SetPropertyTangentPressure)>);
+
+	void AssignSetPropertyTangentPressureData(int8 *const Data, const int32 Extent)
+	{
+		*reinterpret_cast<float*>(Data) = 1.0f / Extent;
+	}
+
+	void SetPropertyOrientation(FStylusInputPacket& Packet, const PACKET& WintabPacket, const int8* Data)
+	{
+		const float& InvScaleAzimuth = reinterpret_cast<const float*>(Data)[0];
+		const float& InvScaleAltitude = reinterpret_cast<const float*>(Data)[1];
+		const float& InvScaleTwist = reinterpret_cast<const float*>(Data)[2];
+
+		Packet.AzimuthOrientation = WintabPacket.pkOrientation.orAzimuth * InvScaleAzimuth;
+		Packet.AltitudeOrientation = WintabPacket.pkOrientation.orAltitude * InvScaleAltitude;
+		Packet.TwistOrientation = WintabPacket.pkOrientation.orTwist * InvScaleTwist;
+
+		const float AzimuthRadians = FMath::DegreesToRadians(Packet.AzimuthOrientation);
+		const float AltitudeRadians = FMath::DegreesToRadians(Packet.AltitudeOrientation);
+
+		const float X = FMath::Sin(AzimuthRadians) * FMath::Cos(AltitudeRadians);
+		const float Y = FMath::Cos(AzimuthRadians) * FMath::Cos(AltitudeRadians);
+
+		Packet.XTiltOrientation = FMath::RadiansToDegrees(FMath::Asin(-X));
+		Packet.YTiltOrientation = FMath::RadiansToDegrees(FMath::Asin(-Y));
+	}
+	static_assert(std::is_same_v<FPacketPropertyHandler::FFuncSetProperty, decltype(SetPropertyOrientation)>);
+
+	void AssignSetPropertyOrientationData(int8 *const Data, const float AzimuthResolution, const float AltitudeResolution, const float TwistResolution)
+	{
+		float* DataFloat = reinterpret_cast<float*>(Data);
+		*DataFloat++ = 1.0f / AzimuthResolution * 360.0f;
+		*DataFloat++ = 1.0f / AltitudeResolution * 360.0f;
+		*DataFloat = 1.0f / TwistResolution * 360.0f;
+	}
+
+	void SetPropertyRotation(FStylusInputPacket& Packet, const PACKET& WintabPacket, const int8* Data)
+	{
+		const float& InvScalePitch = reinterpret_cast<const float*>(Data)[0];
+		const float& InvScaleRoll = reinterpret_cast<const float*>(Data)[1];
+		const float& InvScaleYaw = reinterpret_cast<const float*>(Data)[2];
+
+		Packet.PitchRotation = WintabPacket.pkRotation.roPitch * InvScalePitch;
+		Packet.RollRotation = WintabPacket.pkRotation.roRoll * InvScaleRoll;
+		Packet.YawRotation = WintabPacket.pkRotation.roYaw * InvScaleYaw;
+	}
+	static_assert(std::is_same_v<FPacketPropertyHandler::FFuncSetProperty, decltype(SetPropertyRotation)>);
+
+	void AssignSetPropertyRotationData(int8 *const Data, const float PitchResolution, const float RollResolution, const float YawResolution)
+	{
+		float* DataFloat = reinterpret_cast<float*>(Data);
+		*DataFloat++ = 1.0f / PitchResolution * 360.0f;
+		*DataFloat++ = 1.0f / RollResolution * 360.0f;
+		*DataFloat = 1.0f / YawResolution * 360.0f;
+	}
+
+	bool SetupTabletContextPacketDescriptionData(const FWintabAPI& WintabAPI, UINT WintabContextIndex, const LOGCONTEXT& WintabContext,
+	                                             const RECT& WindowRect, FTabletContext& TabletContext, const FWintabMessageHandler& MessageHandler)
+	{
+		bool bSuccess = true;
+
+		TabletContext.SupportedProperties = ETabletSupportedProperties::None;
+		TabletContext.NumPacketProperties = 0;
+
+		WTPKT DataAvailableForAllCursors;
+		WTPKT DataAvailableForSomeCursors;
+		bSuccess &= WintabAPI.WtInfo(WTI_DEVICES + WintabContextIndex, DVC_PKTDATA, &DataAvailableForAllCursors) == sizeof(DataAvailableForAllCursors);
+		bSuccess &= WintabAPI.WtInfo(WTI_DEVICES + WintabContextIndex, DVC_CSRDATA, &DataAvailableForSomeCursors) == sizeof(DataAvailableForSomeCursors);
+
+		if (!bSuccess)
+		{
+			LogError(LOG_PREAMBLE, FString::Format(TEXT("Failed to query available packet data for Wintab context with index {0}."), {WintabContextIndex}));
+			return false;
+		}
+
+		const WTPKT DataAvailable = DataAvailableForAllCursors | DataAvailableForSomeCursors;
+
+		auto AddProperty = [&TabletContext](EPacketPropertyType Type, ETabletPropertyMetricUnit MetricUnit, int32 Minimum, int32 Maximum,
+		                                    float Resolution) -> FPacketPropertyHandler&
+		{
+			FPacketProperty& Property = TabletContext.PacketProperties[TabletContext.NumPacketProperties];
+			Property.Type = Type;
+			Property.MetricUnit = MetricUnit;
+			Property.Minimum = Minimum;
+			Property.Maximum = Maximum;
+			Property.Resolution = Resolution;
+
+			FPacketPropertyHandler& PropertyHandler = TabletContext.PacketPropertyHandlers[TabletContext.NumPacketProperties];
+
+			++TabletContext.NumPacketProperties;
+
+			return PropertyHandler;
+		};
+
+		if (DataAvailable & PK_STATUS)
+		{
+			TabletContext.SupportedProperties = ETabletSupportedProperties::PacketStatus;
+			FPacketPropertyHandler& PropertyHandler = AddProperty(EPacketPropertyType::Status, ETabletPropertyMetricUnit::Default, 0, 0, 0.0f);
+			PropertyHandler.SetProperty = &SetPropertyStatus;
+		}
+
+		if (DataAvailable & PK_TIME)
+		{
+			TabletContext.SupportedProperties = TabletContext.SupportedProperties | ETabletSupportedProperties::TimerTick;
+			FPacketPropertyHandler& PropertyHandler = AddProperty(EPacketPropertyType::Time, ETabletPropertyMetricUnit::Default, 0, 0, 0.0f);
+			PropertyHandler.SetProperty = &SetPropertyTime;
+		}
+
+		if (DataAvailable & PK_SERIAL_NUMBER)
+		{
+			TabletContext.SupportedProperties = TabletContext.SupportedProperties | ETabletSupportedProperties::SerialNumber;
+			FPacketPropertyHandler& PropertyHandler = AddProperty(EPacketPropertyType::SerialNumber, ETabletPropertyMetricUnit::Default, 0, 0, 0.0f);
+			PropertyHandler.SetProperty = &SetPropertySerialNumber;
+		}
+
+		if (DataAvailable & PK_CURSOR)
+		{
+			FPacketPropertyHandler& PropertyHandler = AddProperty(EPacketPropertyType::Cursor, ETabletPropertyMetricUnit::Default, 0, 0, 0.0f);
+			PropertyHandler.SetProperty = &SetPropertyCursor;
+			AssignSetPropertyCursorData(PropertyHandler.SetPropertyData, &MessageHandler);
+		}
+
+		if (DataAvailable & PK_BUTTONS)
+		{
+			FPacketPropertyHandler& PropertyHandler = AddProperty(EPacketPropertyType::Buttons, ETabletPropertyMetricUnit::Default, 0, 0, 0.0f);
+			PropertyHandler.SetProperty = &SetPropertyButtons;
+		}
+
+		if (DataAvailable & PK_X)
+		{
+			AXIS Axis;
+			if (WintabAPI.WtInfo(WTI_DEVICES + WintabContextIndex, DVC_X, &Axis) == sizeof(Axis))
+			{
+				TabletContext.SupportedProperties = TabletContext.SupportedProperties | ETabletSupportedProperties::X;
+				FPacketPropertyHandler& PropertyHandler = AddProperty(EPacketPropertyType::X, static_cast<ETabletPropertyMetricUnit>(Axis.axUnits), Axis.axMin,
+				                                                      Axis.axMax, Fix32ToFloat(Axis.axResolution));
+				PropertyHandler.SetProperty = &SetPropertyX;
+				AssignSetPropertyXData(PropertyHandler.SetPropertyData, WintabContext.lcInExtX, WintabAPI.GetSystemMetrics(SM_CXVIRTUALSCREEN), &WindowRect);
+			}
+			else
+			{
+				bSuccess = false;
+			}
+		}
+
+		if (DataAvailable & PK_Y)
+		{
+			AXIS Axis;
+			if (WintabAPI.WtInfo(WTI_DEVICES + WintabContextIndex, DVC_Y, &Axis) == sizeof(Axis))
+			{
+				TabletContext.SupportedProperties = TabletContext.SupportedProperties | ETabletSupportedProperties::Y;
+				FPacketPropertyHandler& PropertyHandler = AddProperty(EPacketPropertyType::Y, static_cast<ETabletPropertyMetricUnit>(Axis.axUnits), Axis.axMin,
+				                                                      Axis.axMax, Fix32ToFloat(Axis.axResolution));
+				PropertyHandler.SetProperty = &SetPropertyY;
+				AssignSetPropertyYData(PropertyHandler.SetPropertyData, WintabContext.lcInExtY, WintabAPI.GetSystemMetrics(SM_CYVIRTUALSCREEN), &WindowRect);
+			}
+			else
+			{
+				bSuccess = false;
+			}
+		}
+
+		if (DataAvailable & PK_Z)
+		{
+			AXIS Axis;
+			if (WintabAPI.WtInfo(WTI_DEVICES + WintabContextIndex, DVC_Z, &Axis) == sizeof(Axis))
+			{
+				TabletContext.SupportedProperties = TabletContext.SupportedProperties | ETabletSupportedProperties::Z;
+				FPacketPropertyHandler& PropertyHandler = AddProperty(EPacketPropertyType::Z, static_cast<ETabletPropertyMetricUnit>(Axis.axUnits), Axis.axMin,
+				                                                      Axis.axMax, Fix32ToFloat(Axis.axResolution));
+
+				const auto& Property = TabletContext.PacketProperties[TabletContext.NumPacketProperties - 1]; 
+				const float Resolution = [&Property]
+				{
+					switch (Property.MetricUnit)
+					{
+					case ETabletPropertyMetricUnit::Inches:
+						return Property.Resolution * 2.54f;
+					case ETabletPropertyMetricUnit::Centimeters:
+						return Property.Resolution;
+					default:
+						return 1.0f;
+					}
+				}();
+
+				PropertyHandler.SetProperty = &SetPropertyZ;
+				AssignSetPropertyZData(PropertyHandler.SetPropertyData, Resolution);
+			}
+			else
+			{
+				bSuccess = false;
+			}
+		}
+
+		if (DataAvailable & PK_NORMAL_PRESSURE)
+		{
+			AXIS Axis;
+			if (WintabAPI.WtInfo(WTI_DEVICES + WintabContextIndex, DVC_NPRESSURE, &Axis) == sizeof(Axis))
+			{
+				TabletContext.SupportedProperties = TabletContext.SupportedProperties | ETabletSupportedProperties::NormalPressure;
+				FPacketPropertyHandler& PropertyHandler = AddProperty(EPacketPropertyType::NormalPressure, static_cast<ETabletPropertyMetricUnit>(Axis.axUnits), Axis.axMin,
+				                                        Axis.axMax, Fix32ToFloat(Axis.axResolution));
+				PropertyHandler.SetProperty = &SetPropertyNormalPressure;
+				AssignSetPropertyNormalPressureData(PropertyHandler.SetPropertyData, Axis.axMax - Axis.axMin + 1);
+			}
+			else
+			{
+				bSuccess = false;
+			}
+		}
+
+		if (DataAvailable & PK_TANGENT_PRESSURE)
+		{
+			AXIS Axis;
+			if (WintabAPI.WtInfo(WTI_DEVICES + WintabContextIndex, DVC_TPRESSURE, &Axis) == sizeof(Axis))
+			{
+				TabletContext.SupportedProperties = TabletContext.SupportedProperties | ETabletSupportedProperties::TangentPressure;
+				FPacketPropertyHandler& PropertyHandler = AddProperty(EPacketPropertyType::TangentPressure, static_cast<ETabletPropertyMetricUnit>(Axis.axUnits), Axis.axMin,
+				                                        Axis.axMax, Fix32ToFloat(Axis.axResolution));
+				PropertyHandler.SetProperty = &SetPropertyTangentPressure;
+				AssignSetPropertyTangentPressureData(PropertyHandler.SetPropertyData, Axis.axMax - Axis.axMin + 1);
+			}
+			else
+			{
+				bSuccess = false;
+			}
+		}
+
+		if (DataAvailable & PK_ORIENTATION)
+		{
+			AXIS Orientation[3];
+			if (WintabAPI.WtInfo(WTI_DEVICES + WintabContextIndex, DVC_ORIENTATION, &Orientation) == sizeof(Orientation))
+			{
+				if (Orientation[0].axResolution && Orientation[1].axResolution)
+				{
+					TabletContext.SupportedProperties = TabletContext.SupportedProperties | ETabletSupportedProperties::AzimuthOrientation |
+						ETabletSupportedProperties::AltitudeOrientation | ETabletSupportedProperties::TwistOrientation |
+						ETabletSupportedProperties::XTiltOrientation | ETabletSupportedProperties::YTiltOrientation;
+
+					if (Orientation[0].axUnits != TU_CIRCLE || Orientation[1].axUnits != TU_CIRCLE || Orientation[2].axUnits != TU_CIRCLE)
+					{
+						LogWarning(LOG_PREAMBLE, FString::Format(
+							           TEXT("Units for orientation are not all TU_CIRCLE for Wintab context with index {0}."), {WintabContextIndex}));
+					}
+
+					FPacketPropertyHandler& PropertyHandler = AddProperty(EPacketPropertyType::Orientation, static_cast<ETabletPropertyMetricUnit>(Orientation[0].axUnits), 0,
+					                                        0, Fix32ToFloat(Orientation[0].axResolution));
+					PropertyHandler.SetProperty = &SetPropertyOrientation;
+					AssignSetPropertyOrientationData(PropertyHandler.SetPropertyData, Fix32ToFloat(Orientation[0].axResolution),
+					                                 Fix32ToFloat(Orientation[1].axResolution),
+					                                 Fix32ToFloat(Orientation[2].axResolution));
+				}
+			}
+			else
+			{
+				bSuccess = false;
+			}
+		}
+
+		if (DataAvailable & PK_ROTATION)
+		{
+			AXIS Rotation[3];
+			if (WintabAPI.WtInfo(WTI_DEVICES + WintabContextIndex, DVC_ROTATION, &Rotation) == sizeof(Rotation))
+			{
+				TabletContext.SupportedProperties = TabletContext.SupportedProperties | ETabletSupportedProperties::PitchRotation |
+					ETabletSupportedProperties::RollRotation | ETabletSupportedProperties::YawRotation;
+
+				if (Rotation[0].axUnits != TU_CIRCLE || Rotation[1].axUnits != TU_CIRCLE || Rotation[2].axUnits != TU_CIRCLE)
+				{
+					LogWarning(LOG_PREAMBLE, FString::Format(
+								 TEXT("Units for rotation are not all TU_CIRCLE for Wintab context with index {0}."), {WintabContextIndex}));
+				}
+
+				FPacketPropertyHandler& PropertyHandler = AddProperty(EPacketPropertyType::Rotation, static_cast<ETabletPropertyMetricUnit>(Rotation[0].axUnits), 0, 0,
+				                                        Fix32ToFloat(Rotation[0].axResolution));
+				PropertyHandler.SetProperty = &SetPropertyRotation;
+				AssignSetPropertyRotationData(PropertyHandler.SetPropertyData, Fix32ToFloat(Rotation[0].axResolution), Fix32ToFloat(Rotation[1].axResolution),
+				                              Fix32ToFloat(Rotation[2].axResolution));
+			}
+			else
+			{
+				bSuccess = false;
+			}
+		}
+
+		if (!bSuccess)
+		{
+			LogError(LOG_PREAMBLE, FString::Format(TEXT("Failed to query packet property information for Wintab context with index {0}."), {WintabContextIndex}));
+		}
+
+		return bSuccess;
+	}
+
+	bool SetupTabletContext(const FWintabAPI& WintabAPI, UINT WintabContextIndex, HCTX WintabContextHandle, const LOGCONTEXT& WintabContext,
+	                        const RECT& WindowRect, FTabletContext& TabletContext, const FWintabMessageHandler& MessageHandler)
+	{
+		bool bSuccess = true;
+
+		bSuccess &= SetupTabletContextMetadata(WintabAPI, WintabContextIndex, WintabContextHandle, WintabContext, TabletContext);
+		bSuccess &= SetupTabletContextPacketDescriptionData(WintabAPI, WintabContextIndex, WintabContext, WindowRect, TabletContext, MessageHandler);
+
+		return bSuccess;
+	}
+
+	FWintabInstance::FWintabInstance(const uint32 ID, const HWND OSWindowHandle)
+        : ID(ID)
+		, WintabAPI(FWintabAPI::GetInstance())
+		, OSWindowHandle(OSWindowHandle) 
+		, MessageHandler(this,
+		                 FGetTabletContextCallback::CreateRaw(this, &FWintabInstance::GetTabletContextInternal),
+		                 FGetStylusIDCallback::CreateRaw(this, &FWintabInstance::GetStylusID),
+		                 FUpdateWindowRectCallback::CreateRaw(this, &FWintabInstance::UpdateWindowRect))
+	{
+		if (FSlateApplication::IsInitialized())
+		{
+			if (const TSharedPtr<GenericApplication> WindowsApplication = FSlateApplication::Get().GetPlatformApplication())
+			{
+				static_cast<FWindowsApplication*>(WindowsApplication.Get())->AddMessageHandler(MessageHandler);
+			}
+		}
+
+	    UpdateTabletContexts();
+    }
+
+	FWintabInstance::~FWintabInstance()
+	{
+		for (int32 Index = 0, NumIndices = TabletContexts.Num(); Index < NumIndices; ++Index)
+		{
+			if (!WintabAPI.WtClose(TabletContexts[Index]->WintabContextHandle))
+			{
+				UE_LOG(LogStylusInput, Error, TEXT("WintabInstance: Could not close Wintab context with handle %p."), TabletContexts[Index]->WintabContextHandle);
+			}
+		}
+
+		if (FSlateApplication::IsInitialized())
+		{
+			if (const TSharedPtr<GenericApplication> WindowsApplication = FSlateApplication::Get().GetPlatformApplication())
+			{
+				static_cast<FWindowsApplication*>(WindowsApplication.Get())->RemoveMessageHandler(MessageHandler);
+			}
+		}
+	}
+
+	bool FWintabInstance::AddEventHandler(IStylusInputEventHandler* EventHandler, const EEventHandlerThread Thread)
+	{
+        if (!EventHandler)
+        {
+            LogWarning(LOG_PREAMBLE, "Tried to add nullptr as event handler.");
+            return false;
+        }
+
+		if (Thread == EEventHandlerThread::OnGameThread)
+		{
+			return MessageHandler.AddEventHandler(EventHandler);
+		}
+
+		if (Thread == EEventHandlerThread::Asynchronous)
+		{
+			// TODO Handle asynchronous thread
+			return false;
+		}
+
+        return false;
+	}
+
+	bool FWintabInstance::RemoveEventHandler(IStylusInputEventHandler* EventHandler)
+	{
+        if (!EventHandler)
+        {
+            LogWarning(LOG_PREAMBLE, "Tried to remove nullptr event handler.");
+            return false;
+        }
+
+		if (MessageHandler.RemoveEventHandler(EventHandler))
+		{
+			return true;
+		}
+
+		// TODO Handle asynchronous thread
+
+		return false;
+	}
+
+	const TSharedPtr<IStylusInputTabletContext> FWintabInstance::GetTabletContext(const uint32 TabletContextID)
+	{
+        return TabletContexts.Get(TabletContextID);
+    }
+
+	const TSharedPtr<IStylusInputStylusInfo> FWintabInstance::GetStylusInfo(const uint32 StylusID)
+	{
+		return StylusInfos.Get(StylusID);
+	}
+
+	float FWintabInstance::GetPacketsPerSecond(const EEventHandlerThread EventHandlerThread) const
+	{
+        return MessageHandler.GetPacketsPerSecond();
+    }
+
+	FText FWintabInstance::GetName()
+	{
+		return FText::Format(LOCTEXT("Wintab", "Wintab #{0}"), ID);
+	}
+
+	const FTabletContext& FWintabInstance::GetTabletContextInternal(const uint32 TabletContextID) const
+	{
+		return *TabletContexts.Get(TabletContextID);
+	}
+
+	uint32 FWintabInstance::GetStylusID(const uint32 TabletContextID, uint32 CursorIndex)
+	{
+		DWORD CursorPhysicalID;
+		UINT CursorType;
+		if (WintabAPI.WtInfo(WTI_CURSORS + CursorIndex, CSR_PHYSID, &CursorPhysicalID) == sizeof(CursorPhysicalID) &&
+			WintabAPI.WtInfo(WTI_CURSORS + CursorIndex, CSR_TYPE, &CursorType) == sizeof(CursorType))
+		{
+			const uint16 MaskedCursorId = MaskCursorId(CursorType);
+			const uint16 MaskedCursorType = MaskCursorType(CursorType);
+
+			const uint64 CursorID = static_cast<uint64>(MaskedCursorId) << 32 | CursorPhysicalID;
+			const TTuple<uint64, uint32>* Mapping = CursorIDToStylusIDMappings.FindByPredicate(
+				[CursorID](const TTuple<uint64, uint32>& Tuple)
+				{
+					return Tuple.Get<0>() == CursorID;
+				});
+
+			if (!Mapping)
+			{
+				uint32 StylusID = CursorPhysicalID;
+
+				// Resolve any collisions of physical IDs, which are only guaranteed to be unique within a (masked) cursor type.
+				while (StylusInfos.Contains(StylusID))
+				{
+					++StylusID;
+				}
+
+				Mapping = &CursorIDToStylusIDMappings.Emplace_GetRef(CursorID, StylusID);
+
+				const TSharedRef<FStylusInfo> StylusInfo = StylusInfos.Add(StylusID);
+
+				StylusInfo->WintabPhysicalID = CursorPhysicalID;
+				StylusInfo->WintabCursorType = MaskedCursorType;
+
+				const ECursorIndexType CursorIndexType = [&TabletContexts = TabletContexts, TabletContextID, CursorIndex]
+				{
+					if (const TSharedPtr<FTabletContext> TabletContextPtr = TabletContexts.Get(TabletContextID))
+					{
+						const FTabletContext& TabletContext = *TabletContextPtr;
+						if (TabletContext.WintabFirstCursor <= CursorIndex && CursorIndex < static_cast<uint8>(TabletContext.WintabFirstCursor + TabletContext.WintabNumCursors))
+						{
+							const int8 CursorIndexTypeInt = CursorIndex - TabletContext.WintabFirstCursor;
+							if (0 <= CursorIndexTypeInt && CursorIndexTypeInt < static_cast<int8>(ECursorIndexType::Num_Enumerators))
+							{
+								return static_cast<ECursorIndexType>(CursorIndexTypeInt);
+							}
+						}
+					}
+					return ECursorIndexType::Invalid_Enumerator;
+				}();
+
+				// If this is the inverted side of a pen, populate the stylus data with the non-inverted side instead.
+				if (CursorIndexType != ECursorIndexType::Invalid_Enumerator && CursorIsInverted(CursorIndexType))
+				{
+					--CursorIndex;
+				}
+
+				FWintabInfoOutputBuffer OutputBuffer;
+				OutputBuffer.Allocate(WTI_CURSORS + CursorIndex, CSR_NAME);
+				if (WintabAPI.WtInfo(WTI_CURSORS + CursorIndex, CSR_NAME, OutputBuffer.Get()) <= OutputBuffer.SizeInBytes())
+				{
+					StylusInfo->Name = FString::Format(TEXT("{0} ({1})"), {OutputBuffer.Get(), MaskedCursorTypeToString(MaskedCursorType)});
+				}
+				else
+				{
+					LogWarning(LOG_PREAMBLE, FString::Format(TEXT("Failed to query name for Wintab cursor with index {0}."), {CursorIndex}));
+				}
+
+				BYTE NumButtons;
+				if (WintabAPI.WtInfo(WTI_CURSORS + CursorIndex, CSR_BUTTONS, &NumButtons) == sizeof(NumButtons))
+				{
+					StylusInfo->Buttons.SetNum(NumButtons);
+
+					OutputBuffer.Allocate(WTI_CURSORS + CursorIndex, CSR_BTNNAMES);
+					if (WintabAPI.WtInfo(WTI_CURSORS + CursorIndex, CSR_BTNNAMES, OutputBuffer.Get()) <= OutputBuffer.SizeInBytes())
+					{
+						const TCHAR* Name = OutputBuffer.Get();
+						int32 ButtonIndex = 0;
+						while (*Name)
+						{
+							StylusInfo->Buttons[ButtonIndex].Name = Name;
+							Name += StylusInfo->Buttons[ButtonIndex].Name.Len() + 1;
+							++ButtonIndex;
+						}
+					}
+					else
+					{
+						LogWarning(LOG_PREAMBLE, FString::Format(TEXT("Failed to query button names for Wintab cursor with index {0}."), {CursorIndex}));
+					}
+
+					BYTE ButtonMap[32] = {};
+					if (WintabAPI.WtInfo(WTI_CURSORS + CursorIndex, CSR_BUTTONMAP, &ButtonMap) == sizeof(ButtonMap))
+					{
+						for (int32 ButtonIndex = 0; ButtonIndex < NumButtons; ++ButtonIndex)
+						{
+							// Todo
+						}
+					}
+					else
+					{
+						LogWarning(LOG_PREAMBLE, FString::Format(TEXT("Failed to query button map for Wintab cursor with index {0}."), {CursorIndex}));
+					}
+				}
+				else
+				{
+					LogWarning(LOG_PREAMBLE, FString::Format(TEXT("Failed to query number of buttons for Wintab cursor with index {0}."), {CursorIndex}));
+				}
+			}
+
+			// Make sure a StylusInfo with this ID exists, and the cursor type and physical ID match.
+			checkSlow(Mapping && StylusInfos.Contains(Mapping->Get<1>()) && StylusInfos.Get(Mapping->Get<1>()).IsValid()
+				&& StylusInfos.Get(Mapping->Get<1>())->WintabCursorType == MaskedCursorType
+				&& StylusInfos.Get(Mapping->Get<1>())->WintabPhysicalID == CursorPhysicalID);
+
+			return Mapping->Get<1>();
+		}
+		else
+		{
+			LogWarning(LOG_PREAMBLE, FString::Format(TEXT("Failed to query cursor metadata for cursor index {0}."), {CursorIndex}));
+		}
+
+		return 0;
+	}
+
+	void FWintabInstance::UpdateTabletContexts()
+	{
+		UpdateWindowRect();
+
+		UINT WintabContextIndex = 0;
+		LOGCONTEXTW WintabContext;
+
+		while (WintabAPI.WtInfo(WTI_DDCTXS + WintabContextIndex, 0, &WintabContext) == sizeof(WintabContext))
+		{
+			if (!SetupWintabContext(WintabAPI, WintabContextIndex, WintabContext))
+			{
+				LogError(LOG_PREAMBLE, FString::Format(TEXT("Failed to setup Wintab context with index {0}}."), {WintabContextIndex}));
+				continue;
+			}
+
+			const HCTX WintabContextHandle = WintabAPI.WtOpen(OSWindowHandle, &WintabContext, Windows::TRUE);
+			if (!WintabContextHandle)
+			{
+				LogError(LOG_PREAMBLE, FString::Format(TEXT("Failed to open Wintab context with index {0}."), {WintabContextIndex}));
+				continue;
+			}
+
+			const uint32 TabletContextID = HctxToUint32(WintabContextHandle);
+
+			TSharedRef<FTabletContext> TabletContextPtr = [&TabletContexts = TabletContexts, TabletContextID]
+			{
+				if (TabletContexts.Contains(TabletContextID))
+				{
+					return TabletContexts.Get(TabletContextID).ToSharedRef();
+				}
+				return TabletContexts.Add(TabletContextID);
+			}();
+
+			FTabletContext& TabletContext = *TabletContextPtr;
+
+			if (!SetupTabletContext(WintabAPI, WintabContextIndex, WintabContextHandle, WintabContext, WindowRect, TabletContext, MessageHandler))
+			{
+				LogError(LOG_PREAMBLE, FString::Format(TEXT("Failed setting up tablet context with ID {0}."), {TabletContextID}));
+
+				TabletContexts.Remove(TabletContextID);
+
+				if (!WintabAPI.WtClose(WintabContextHandle))
+				{
+					LogError(LOG_PREAMBLE, FString::Format(TEXT("Could not close Wintab context with handle {0}."), {HctxToUint32(WintabContextHandle)}));
+				}
+
+				continue;
+			}
+
+			LogVerbose(LOG_PREAMBLE, FString::Format(TEXT("Added TabletContext for ID {0} [{1}, {2}]."), {
+				                                         TabletContext.ID, TabletContext.Name , TabletContext.PlugAndPlayID
+			                                         }));
+
+			++WintabContextIndex;
+		}
+	}
+
+	void FWintabInstance::UpdateWindowRect()
+	{
+		if (!ensure(WintabAPI.GetWindowRect(OSWindowHandle, &WindowRect)))
+		{
+			LogError(LOG_PREAMBLE, "Could not retrieve window rectangle; coordinates mapping will be incorrect!");
+		}
+	}
+}
+
+#undef LOG_PREAMBLE
+#undef LOCTEXT_NAMESPACE
+
+#endif
diff --git a/Engine/Plugins/Editor/StylusInput/Source/StylusInput/Private/Windows/Wintab/WintabInstance.h b/Engine/Plugins/Editor/StylusInput/Source/StylusInput/Private/Windows/Wintab/WintabInstance.h
new file mode 100644
--- /dev/null
+++ b/Engine/Plugins/Editor/StylusInput/Source/StylusInput/Private/Windows/Wintab/WintabInstance.h
@@ -0,0 +1,62 @@
+﻿// Copyright Epic Games, Inc. All Rights Reserved.
+
+#pragma once
+
+#include <HAL/Platform.h>
+
+#if PLATFORM_WINDOWS && STYLUSINPUT_USE_WINTAB
+
+#include <StylusInput.h>
+#include <Containers/Array.h>
+#include <Microsoft/MinimalWindowsApi.h>
+#include <Templates/UniquePtr.h>
+
+#include "WintabMessageHandler.h"
+#include "WintabTabletContext.h"
+#include "WintabStylus.h"
+
+namespace UE::StylusInput
+{
+	class IStylusInputEventHandler;
+}
+
+namespace UE::StylusInput::Private::Wintab
+{
+	class FWintabAPI;
+
+	class FWintabInstance : public IStylusInputInstance
+	{
+	public:
+		explicit FWintabInstance(uint32 ID, HWND OSWindowHandle);
+		virtual ~FWintabInstance() override;
+
+		virtual bool AddEventHandler(IStylusInputEventHandler* EventHandler, EEventHandlerThread Thread) override;
+		virtual bool RemoveEventHandler(IStylusInputEventHandler* EventHandler) override;
+
+		virtual const TSharedPtr<IStylusInputTabletContext> GetTabletContext(uint32 TabletContextID) override;
+		virtual const TSharedPtr<IStylusInputStylusInfo> GetStylusInfo(uint32 StylusID) override;
+
+		virtual float GetPacketsPerSecond(EEventHandlerThread EventHandlerThread) const override;
+
+		virtual FText GetName() override;
+
+	private:
+		const FTabletContext& GetTabletContextInternal(uint32 TabletContextID) const;
+		uint32 GetStylusID(uint32 TabletContextID, uint32 CursorIndex);
+		void UpdateTabletContexts();
+		void UpdateWindowRect();
+
+		const uint32 ID;
+		const FWintabAPI& WintabAPI;
+		const HWND OSWindowHandle;
+		RECT WindowRect = {};
+
+		FTabletContextContainer TabletContexts;
+		FStylusInfoContainer StylusInfos;
+
+		FWintabMessageHandler MessageHandler;
+		TArray<TTuple<uint64, uint32>> CursorIDToStylusIDMappings;
+	};
+}
+
+#endif
diff --git a/Engine/Plugins/Editor/StylusInput/Source/StylusInput/Private/Windows/Wintab/WintabInterface.cpp b/Engine/Plugins/Editor/StylusInput/Source/StylusInput/Private/Windows/Wintab/WintabInterface.cpp
new file mode 100644
--- /dev/null
+++ b/Engine/Plugins/Editor/StylusInput/Source/StylusInput/Private/Windows/Wintab/WintabInterface.cpp
@@ -0,0 +1,93 @@
+﻿// Copyright Epic Games, Inc. All Rights Reserved.
+
+#include "WintabInterface.h"
+
+#if PLATFORM_WINDOWS && STYLUSINPUT_USE_WINTAB
+
+#include <StylusInput.h>
+#include <StylusInputUtils.h>
+
+#include <GenericPlatform/GenericWindow.h>
+#include <Templates/SharedPointer.h>
+#include <Widgets/SWindow.h>
+
+#include "WintabInstance.h"
+#include "WintabAPI.h"
+
+#define LOCTEXT_NAMESPACE "WintabInterface"
+
+namespace UE::StylusInput::Private::Wintab
+{
+	IStylusInputInstance* FWintabInterface::CreateInstance(SWindow& Window)
+	{
+		if (FRefCountedInstance* ExistingRefCountedInstance = Instances.Find(&Window))
+		{
+			++ExistingRefCountedInstance->RefCount;
+			return ExistingRefCountedInstance->Instance.Get();
+		}
+
+		HWND OSWindowHandle = [&Window]
+		{
+			const TSharedPtr<const FGenericWindow> NativeWindow = Window.GetNativeWindow();
+			return NativeWindow.IsValid() ? static_cast<HWND>(NativeWindow->GetOSWindowHandle()) : nullptr;
+		}();
+
+		if (!OSWindowHandle)
+		{
+			LogError("WintabInterface", "Could not get native window handle.");
+			return nullptr;
+		}
+
+		FWintabInstance* NewInstance = Instances.Emplace(
+			&Window, {MakeUnique<FWintabInstance>(NextInstanceID++, OSWindowHandle), 1}).Instance.Get();
+		if (!ensureMsgf(NewInstance, TEXT("WintabInterface: Failed to create stylus input instance.")))
+		{
+			Instances.Remove(&Window);
+			return nullptr;
+		}
+
+		return NewInstance;
+	}
+
+	bool FWintabInterface::ReleaseInstance(IStylusInputInstance* Instance)
+	{
+		check(Instance);
+
+		FWintabInstance *const WindowsInstance = static_cast<FWintabInstance*>(Instance);
+
+		// Find existing instance
+		for (TTuple<SWindow*, FRefCountedInstance>& Entry : Instances)
+		{
+			FRefCountedInstance& RefCountedInstance = Entry.Get<1>();
+
+			if (RefCountedInstance.Instance.Get() == WindowsInstance)
+			{
+				// Decrease reference count
+				check(RefCountedInstance.RefCount > 0);
+				if (--RefCountedInstance.RefCount == 0)
+				{
+					// Delete if there are no references left
+					Instances.Remove(Entry.Key);
+				}
+
+				return true;
+			}
+		}
+
+		ensureMsgf(false, TEXT("WintabInterface: Failed to find provided instance."));
+		return false;
+	}
+
+	TUniquePtr<IStylusInputInterface> FWintabInterface::Create()
+	{
+		if (FWintabAPI::GetInstance().SatisfiesRequirements())
+		{
+			TUniquePtr<FWintabInterface> WindowsImpl = MakeUnique<FWintabInterface>();
+			return TUniquePtr<IStylusInputInterface>(MoveTemp(WindowsImpl));
+		}
+
+		return nullptr;
+	}
+}
+
+#endif
diff --git a/Engine/Plugins/Editor/StylusInput/Source/StylusInput/Private/Windows/Wintab/WintabInterface.h b/Engine/Plugins/Editor/StylusInput/Source/StylusInput/Private/Windows/Wintab/WintabInterface.h
new file mode 100644
--- /dev/null
+++ b/Engine/Plugins/Editor/StylusInput/Source/StylusInput/Private/Windows/Wintab/WintabInterface.h
@@ -0,0 +1,45 @@
+﻿// Copyright Epic Games, Inc. All Rights Reserved.
+
+#pragma once
+
+#include <HAL/Platform.h>
+
+#if PLATFORM_WINDOWS && STYLUSINPUT_USE_WINTAB
+
+#include <StylusInputInterface.h>
+#include <Containers/Array.h>
+#include <Containers/Map.h>
+#include <Templates/UniquePtr.h>
+
+namespace UE::StylusInput::Private
+{
+	class FStylusInputImpl;
+}
+
+namespace UE::StylusInput::Private::Wintab
+{
+	class FWintabInstance;
+
+	class FWintabInterface : public IStylusInputInterface
+	{
+	public:
+		virtual IStylusInputInstance* CreateInstance(SWindow& Window) override;
+		virtual bool ReleaseInstance(IStylusInputInstance* Instance) override;
+
+	private:
+		friend FStylusInputImpl;
+
+		static TUniquePtr<IStylusInputInterface> Create();
+
+		struct FRefCountedInstance
+		{
+			TUniquePtr<FWintabInstance> Instance;
+			int32 RefCount;
+		};
+
+		TMap<SWindow*, FRefCountedInstance> Instances;
+		uint32 NextInstanceID = 0;
+	};
+}
+
+#endif
diff --git a/Engine/Plugins/Editor/StylusInput/Source/StylusInput/Private/Windows/Wintab/WintabMessageHandler.cpp b/Engine/Plugins/Editor/StylusInput/Source/StylusInput/Private/Windows/Wintab/WintabMessageHandler.cpp
new file mode 100644
--- /dev/null
+++ b/Engine/Plugins/Editor/StylusInput/Source/StylusInput/Private/Windows/Wintab/WintabMessageHandler.cpp
@@ -0,0 +1,273 @@
+﻿// Copyright Epic Games, Inc. All Rights Reserved.
+
+#include "WintabMessageHandler.h"
+
+#if PLATFORM_WINDOWS && STYLUSINPUT_USE_WINTAB
+
+#include <StylusInputPacket.h>
+#include <StylusInputUtils.h>
+
+#include "WintabAPI.h"
+#include "WintabTabletContext.h"
+
+#define LOG_PREAMBLE "WintabMessageHandler"
+
+#define ENABLE_DEBUG_EVENTS_FOR_PACKET_MESSAGES 0
+#define ENABLE_DEBUG_EVENTS_FOR_CONTEXT_OPEN_CLOSE_UPDATE_MESSAGES 0
+#define ENABLE_DEBUG_EVENTS_FOR_CONTEXT_OVERLAP_MESSAGES 0
+#define ENABLE_DEBUG_EVENTS_FOR_PROXIMITY_CSRCHANGE_MESSAGES 0
+#define ENABLE_DEBUG_EVENTS_FOR_INFOCHANGE_MESSAGES 0
+#define ENABLE_DEBUG_EVENTS_FOR_PACKETEXT_MESSAGES 0
+
+namespace UE::StylusInput::Private::Wintab
+{
+FWintabMessageHandler::FWintabMessageHandler(IStylusInputInstance* Instance, FGetTabletContextCallback&& GetTabletContextCallback,
+                                             FGetStylusIDCallback&& GetStylusIDCallback, FUpdateWindowRectCallback&& UpdateWindowRectCallback)
+	: Instance(Instance)
+	, WintabAPI(FWintabAPI::GetInstance())
+	, GetTabletContextCallback(MoveTemp(GetTabletContextCallback))
+	, GetStylusIDCallback(MoveTemp(GetStylusIDCallback))
+	, UpdateWindowRectCallback(MoveTemp(UpdateWindowRectCallback))
+{
+}
+
+FWintabMessageHandler::~FWintabMessageHandler()
+{
+}
+
+bool FWintabMessageHandler::ProcessMessage(const HWND Hwnd, const uint32 Msg, const WPARAM WParam, const LPARAM LParam, int32& OutResult)
+{
+	bool bHandled = false;
+
+	switch (Msg)
+	{
+	case WT_PACKET:
+		{
+#if ENABLE_DEBUG_EVENTS_FOR_PACKET_MESSAGES
+			DebugEvent(FString::Format(TEXT("Received WT_PACKET(SerialNumber={0}, Context={1}) message."), {static_cast<uint32>(WParam), static_cast<uint32>(LParam)}));
+#endif
+
+			if (WParam != 0)
+			{
+				PACKET WintabPacket = {};
+				if (WintabAPI.WtPacket(reinterpret_cast<HCTX>(LParam), WParam, &WintabPacket))
+				{
+					// Todo Process more than one packet? 
+
+					PacketStats.NewPacket(WintabPacket.pkSerialNumber);
+
+					const FStylusInputPacket Packet = ProcessPacket(static_cast<uint32>(LParam), &WintabPacket);
+
+					for (IStylusInputEventHandler* EventHandler : EventHandlers)
+					{
+						EventHandler->OnPacket(Packet, Instance);
+					}
+				}
+				else
+				{
+					PacketStats.InvalidPacket();
+					DebugEvent(FString::Format(TEXT("Failed to retrieve Wintab packet number {0}."), {WParam}));
+				}
+			}
+
+			bHandled = true;
+			break;
+		}
+	case WT_CTXOPEN:
+		{
+#if ENABLE_DEBUG_EVENTS_FOR_CONTEXT_OPEN_CLOSE_UPDATE_MESSAGES
+			DebugEvent("Received WT_CTXOPEN message.");
+#endif
+
+			bHandled = true;
+			break;
+		}
+	case WT_CTXCLOSE:
+		{
+#if ENABLE_DEBUG_EVENTS_FOR_CONTEXT_OPEN_CLOSE_UPDATE_MESSAGES
+			DebugEvent("Received WT_CTXCLOSE message.");
+#endif
+
+			bHandled = true;
+			break;
+		}
+	case WT_CTXUPDATE:
+		{
+#if ENABLE_DEBUG_EVENTS_FOR_CONTEXT_OPEN_CLOSE_UPDATE_MESSAGES
+			DebugEvent("Received WT_CTXUPDATE message.");
+#endif
+
+			bHandled = true;
+			break;
+		}
+	case WT_CTXOVERLAP:
+		{
+#if ENABLE_DEBUG_EVENTS_FOR_CONTEXT_OVERLAP_MESSAGES
+			DebugEvent("Received WT_CTXOVERLAP message.");
+#endif
+
+			bHandled = true;
+			break;
+		}
+	case WT_PROXIMITY:
+		{
+#if ENABLE_DEBUG_EVENTS_FOR_PROXIMITY_CSRCHANGE_MESSAGES
+			DebugEvent(FString::Format(TEXT("Received WT_PROXIMITY message (Context={0}, ContextEnter={1}, ProximityEnter={2})."), {
+											   static_cast<uint32>(WParam), LOWORD(LParam), HIWORD(LParam)
+										   }));
+#endif
+
+			if (LOWORD(LParam) || HIWORD(LParam))
+			{
+				// While the Wintab documentation is somewhat confusing about what is encoded in LParam, it appears that the high-order word is effectively
+				// indicating that the event is triggered by hardware proximity, and the low-order word is triggered both with and without hardware proximity
+				// support. Just in case, we trigger a cursor change if either of the flags are set.
+
+				bCursorChange = true;
+				CursorChangeTableContextID = static_cast<uint32>(WParam);
+			}
+
+			bHandled = true;
+			break;
+		}
+	case WT_INFOCHANGE:
+		{
+#if ENABLE_DEBUG_EVENTS_FOR_INFOCHANGE_MESSAGES
+			DebugEvent(FString::Format(TEXT("Received WT_INFOCHANGE(Context={0}, Category={1}, Index={2}) message."), {WParam, LOWORD(LParam), HIWORD(LParam)}));
+#endif
+
+			bHandled = true;
+			break;
+		}
+	case WT_CSRCHANGE:
+		{
+#if ENABLE_DEBUG_EVENTS_FOR_PROXIMITY_CSRCHANGE_MESSAGES
+			DebugEvent(FString::Format(TEXT("Received WT_CSRCHANGE message (SerialNumber={0}, Context={1})."), {static_cast<uint32>(WParam), static_cast<uint32>(LParam)}));
+#endif
+
+			bCursorChange = true;
+			CursorChangeTableContextID = static_cast<uint32>(LParam);
+			CursorChangeSerialNumber = static_cast<uint32>(WParam);
+
+			bHandled = true;
+			break;
+		}
+	case WT_PACKETEXT:
+		{
+#if ENABLE_DEBUG_EVENTS_FOR_PACKETEXT_MESSAGES
+			DebugEvent("Received WT_PACKETEXT message.");
+#endif
+
+			bHandled = true;
+			break;
+		}
+	case WM_WINDOWPOSCHANGED:
+		{
+			UpdateWindowRectCallback.Execute();
+			break;
+		}
+	default:
+		;
+	}
+
+	return bHandled;
+}
+
+bool FWintabMessageHandler::AddEventHandler(IStylusInputEventHandler* EventHandler)
+{
+	check(EventHandler);
+
+	if (EventHandlers.Contains(EventHandler))
+	{
+		LogWarning(LOG_PREAMBLE, FString::Format(TEXT("Event handler '{0}' already exists."), {EventHandler->GetName()}));
+		return false;
+	}
+
+	EventHandlers.Add(EventHandler);
+
+	LogVerbose(LOG_PREAMBLE, FString::Format(TEXT("Event handler '{0}' was added."), {EventHandler->GetName()}));
+
+	return true;
+}
+
+bool FWintabMessageHandler::RemoveEventHandler(IStylusInputEventHandler* EventHandler)
+{
+	check(EventHandler);
+
+	const bool bWasRemoved = EventHandlers.Remove(EventHandler) > 0;
+
+	if (bWasRemoved)
+	{
+		LogVerbose(LOG_PREAMBLE, FString::Format(TEXT("Event handler '{0}' was removed."), {EventHandler->GetName()}));
+	}
+
+	return bWasRemoved;
+}
+
+void FWintabMessageHandler::DebugEvent(const FString& Message) const
+{
+	for (IStylusInputEventHandler* EventHandler : EventHandlers)
+	{
+		EventHandler->OnDebugEvent(Message, Instance);
+	}
+}
+
+const FPacketProperty* FWintabMessageHandler::GetPacketDescriptions(uint32 TabletContextID) const
+{
+	const FTabletContext& TabletContext = GetTabletContextCallback.Execute(TabletContextID);
+	return /*TabletContext ?*/ TabletContext.PacketProperties /*: nullptr*/;
+}
+
+void FWintabMessageHandler::ProcessCursorChange(const uint32 TabletContextID, const void *const WintabPacketPtr)
+{
+	const PACKET& WintabPacket = *static_cast<const PACKET*>(WintabPacketPtr);
+
+	check(bCursorChange);
+
+#if CHECK_MESSAGE_ORDER_ASSUMPTIONS
+	if (CursorChangeTableContextID != TabletContextID || !(CursorChangeSerialNumber == 0 || CursorChangeSerialNumber == WintabPacket.pkSerialNumber))
+	{
+		DebugEvent("Cursor change assumptions violated: The next packet is not matching the previous WT_PROXIMITY/WT_CSRCHANGE messages.");
+	}
+#endif
+
+	CurrentStylusID = GetStylusIDCallback.Execute(TabletContextID, WintabPacket.pkCursor);
+
+	bCursorChange = false;
+
+#if CHECK_MESSAGE_ORDER_ASSUMPTIONS
+	CursorChangeTableContextID = 0;
+	CursorChangeSerialNumber = 0;
+#endif
+}
+
+FStylusInputPacket FWintabMessageHandler::ProcessPacket(const uint32 TabletContextID, const void *const WintabPacketPtr)
+{
+	const PACKET& WintabPacket = *static_cast<const PACKET*>(WintabPacketPtr);
+
+	if (bCursorChange)
+	{
+		ProcessCursorChange(TabletContextID, WintabPacketPtr);
+	}
+
+	const FTabletContext& TabletContext = GetTabletContextCallback.Execute(TabletContextID);
+
+	FStylusInputPacket Packet{TabletContext.ID};
+
+	const bool bPacketOnDigitizer = !static_cast<bool>(TabletContext.SupportedProperties & ETabletSupportedProperties::Z) || WintabPacket.pkZ == 0;
+	Packet.Type = bPacketOnDigitizer
+					  ? bLastPacketOnDigitizer ? EPacketType::OnDigitizer : EPacketType::StylusDown
+					  : bLastPacketOnDigitizer ? EPacketType::StylusUp : EPacketType::AboveDigitizer;
+	bLastPacketOnDigitizer = bPacketOnDigitizer;
+
+	for (int32 Index = 0, Num = TabletContext.NumPacketProperties; Index < Num; ++Index)
+	{
+		const FPacketPropertyHandler& PropertyHandler = TabletContext.PacketPropertyHandlers[Index];
+		PropertyHandler.SetProperty(Packet, WintabPacket, PropertyHandler.SetPropertyData);
+	}
+
+	return Packet;
+}
+}
+
+#endif
diff --git a/Engine/Plugins/Editor/StylusInput/Source/StylusInput/Private/Windows/Wintab/WintabMessageHandler.h b/Engine/Plugins/Editor/StylusInput/Source/StylusInput/Private/Windows/Wintab/WintabMessageHandler.h
new file mode 100644
--- /dev/null
+++ b/Engine/Plugins/Editor/StylusInput/Source/StylusInput/Private/Windows/Wintab/WintabMessageHandler.h
@@ -0,0 +1,71 @@
+﻿// Copyright Epic Games, Inc. All Rights Reserved.
+
+#pragma once
+
+#include <HAL/Platform.h>
+
+#if PLATFORM_WINDOWS && STYLUSINPUT_USE_WINTAB
+
+#include <StylusInput.h>
+#include <Windows/WindowsApplication.h>
+
+#include "WintabStats.h"
+
+#define CHECK_MESSAGE_ORDER_ASSUMPTIONS NDEBUG
+
+namespace UE::StylusInput::Private::Wintab
+{
+	struct FPacketProperty;
+	class FTabletContext;
+	class FStylusInfo;
+	struct FWindowContext;
+	class FWintabAPI;
+
+	DECLARE_DELEGATE_RetVal(const FWindowContext&, FGetWindowContextCallback);
+	DECLARE_DELEGATE_RetVal_OneParam(const FTabletContext&, FGetTabletContextCallback, uint32);
+	DECLARE_DELEGATE_RetVal_TwoParams(uint32, FGetStylusIDCallback, uint32, UINT);
+	DECLARE_DELEGATE(FUpdateWindowRectCallback);
+
+	class FWintabMessageHandler final : public IWindowsMessageHandler
+	{
+	public:
+		FWintabMessageHandler(IStylusInputInstance* Instance, FGetTabletContextCallback&& GetTabletContextCallback, FGetStylusIDCallback&& GetStylusIDCallback,
+		                      FUpdateWindowRectCallback&& UpdateWindowRectCallback);
+		virtual ~FWintabMessageHandler();
+
+		virtual bool ProcessMessage(HWND Hwnd, uint32 Msg, WPARAM WParam, LPARAM LParam, int32& OutResult) override;
+
+		bool AddEventHandler(IStylusInputEventHandler* EventHandler);
+		bool RemoveEventHandler(IStylusInputEventHandler* EventHandler);
+
+		uint32 GetCurrentStylusID() const { return CurrentStylusID; }
+		float GetPacketsPerSecond() const { return PacketStats.GetPacketsPerSecond(); }
+
+	private:
+		void DebugEvent(const FString& Message) const;
+		const FPacketProperty* GetPacketDescriptions(uint32 TabletContextID) const;
+
+		void ProcessCursorChange(uint32 TabletContextID, const void* WintabPacketPtr);
+		FStylusInputPacket ProcessPacket(uint32 TabletContextID, const void* WintabPacketPtr);
+
+		IStylusInputInstance *const Instance;
+		const FWintabAPI& WintabAPI;
+		FGetTabletContextCallback GetTabletContextCallback;
+		FGetStylusIDCallback GetStylusIDCallback;
+		FUpdateWindowRectCallback UpdateWindowRectCallback;
+		TArray<IStylusInputEventHandler*> EventHandlers;
+		FPacketStats PacketStats;
+
+		uint32 CurrentStylusID = 0;
+
+		bool bCursorChange = false;
+		bool bLastPacketOnDigitizer = false;
+
+#if CHECK_MESSAGE_ORDER_ASSUMPTIONS
+		uint32 CursorChangeTableContextID = 0;
+		uint32 CursorChangeSerialNumber = 0;
+#endif
+	};
+}
+
+#endif
diff --git a/Engine/Plugins/Editor/StylusInput/Source/StylusInput/Private/Windows/Wintab/WintabStats.h b/Engine/Plugins/Editor/StylusInput/Source/StylusInput/Private/Windows/Wintab/WintabStats.h
new file mode 100644
--- /dev/null
+++ b/Engine/Plugins/Editor/StylusInput/Source/StylusInput/Private/Windows/Wintab/WintabStats.h
@@ -0,0 +1,72 @@
+﻿// Copyright Epic Games, Inc. All Rights Reserved.
+
+#pragma once
+
+#include <HAL/Platform.h>
+
+#if PLATFORM_WINDOWS && STYLUSINPUT_USE_WINTAB
+
+#include <Windows/WindowsPlatformTime.h>
+
+namespace UE::StylusInput::Private::Wintab
+{
+	class FPacketStats
+	{
+	public:
+		void NewPacket(uint32 SerialNumber)
+		{
+			const uint32 Timestamp = FPlatformTime::Cycles();
+
+			if (LatestTimestamp - EarliestTimestamp >= CyclesPerSecond)
+			{
+				PacketsPerSecond = EarliestTimestamp != LatestTimestamp ? NumPackets / ((LatestTimestamp - EarliestTimestamp) / CyclesPerSecond) : 0.0f;
+				EarliestTimestamp = Timestamp;
+				NumPackets = -1;
+			}
+			LatestTimestamp = Timestamp;
+
+			NumLostPackets += SerialNumber > LatestSerialNumber ? SerialNumber - LatestSerialNumber - 1 : 0;
+			LatestSerialNumber = SerialNumber;
+
+			++NumPackets;
+		}
+
+		void InvalidPacket()
+		{
+			++NumInvalidPackets;
+		}
+
+		void LostPackets()
+		{
+			++NumLostPackets;
+		}
+
+		float GetPacketsPerSecond() const
+		{
+			const uint32 Timestamp = FPlatformTime::Cycles();
+			return Timestamp - LatestTimestamp <= CyclesPerSecond ? PacketsPerSecond : 0.0f;
+		}
+
+		uint32 GetInvalidPackets() const
+		{
+			return NumInvalidPackets;
+		}
+
+		uint32 GetLostPackets() const
+		{
+			return NumLostPackets;
+		}
+
+	private:
+		uint32 EarliestTimestamp = 0;
+		uint32 LatestTimestamp = 0;
+		uint32 LatestSerialNumber = -1;
+		uint32 NumPackets = 0;
+		uint32 NumInvalidPackets = 0;
+		uint32 NumLostPackets = 0;
+		float PacketsPerSecond = 0.0f;
+		const double CyclesPerSecond = 1.0 / FPlatformTime::GetSecondsPerCycle();
+	};
+}
+
+#endif
diff --git a/Engine/Plugins/Editor/StylusInput/Source/StylusInput/Private/Windows/Wintab/WintabStylus.h b/Engine/Plugins/Editor/StylusInput/Source/StylusInput/Private/Windows/Wintab/WintabStylus.h
new file mode 100644
--- /dev/null
+++ b/Engine/Plugins/Editor/StylusInput/Source/StylusInput/Private/Windows/Wintab/WintabStylus.h
@@ -0,0 +1,114 @@
+﻿// Copyright Epic Games, Inc. All Rights Reserved.
+
+#pragma once
+
+#include <HAL/Platform.h>
+
+#if PLATFORM_WINDOWS && STYLUSINPUT_USE_WINTAB
+
+#include <StylusInputTabletContext.h>
+#include <StylusInputUtils.h>
+
+namespace UE::StylusInput::Private::Wintab
+{
+	enum class ECursorType : uint16
+	{
+		GeneralStylus = 0x0802,
+		Airbrush = 0x0902,
+		ArtPen = 0x0804,
+		FourDMouse = 0x0004,
+		FiveButtonPuck = 0x0006,
+
+		CursorTypeMask = 0x0F06,	// Use this to mask any values coming from the Wintab API to get a valid cursor type.
+		CursorIdMask = 0xFF6		// Use this to mask any values coming from the Wintab API to differentiate between different devices, incl. colors. 
+	};
+
+	inline uint16 MaskCursorType(const UINT CursorType)
+	{
+		return static_cast<uint16>(CursorType) & static_cast<uint16>(ECursorType::CursorTypeMask);
+	}
+
+	inline uint16 MaskCursorId(const UINT CursorType)
+	{
+		return static_cast<uint16>(CursorType) & static_cast<uint16>(ECursorType::CursorIdMask);
+	}
+
+	inline FString MaskedCursorTypeToString(uint16 MaskedCursorType)
+	{
+		switch (static_cast<ECursorType>(MaskedCursorType))
+		{
+		case ECursorType::GeneralStylus:
+			return TEXT("General Stylus");
+		case ECursorType::Airbrush:
+			return TEXT("Airbrush");
+		case ECursorType::ArtPen:
+			return TEXT("Art Pen");
+		case ECursorType::FourDMouse:
+			return TEXT("4D Mouse");
+		case ECursorType::FiveButtonPuck:
+			return TEXT("5 Button Puck");
+		default:
+			return TEXT("Unknown Type");
+		}
+	}
+
+	enum class ECursorIndexType : int8
+	{
+		PuckDevice1 = 0,
+		StylusDevice1 = 1,
+		InvertedStylusDevice1 = 2,
+		PuckDevice2 = 3,
+		StylusDevice2 = 4,
+		InvertedStylusDevice2 = 5,
+
+		Num_Enumerators,
+		Invalid_Enumerator = INDEX_NONE
+	};
+
+	inline bool CursorIsInverted(ECursorIndexType Type)
+	{
+		return Type == ECursorIndexType::InvertedStylusDevice1 || Type == ECursorIndexType::InvertedStylusDevice2;
+	}
+
+	class FStylusButton final : public IStylusInputStylusButton
+	{
+	public:
+		virtual FString GetID() const override { return ID; }
+		virtual FString GetName() const override { return Name; }
+
+		FString ID;
+		FString Name;
+	};
+
+	class FStylusInfo final : public IStylusInputStylusInfo
+	{
+	public:
+		FStylusInfo() = default;
+
+		explicit FStylusInfo(const uint32 ID)
+			: ID(ID)
+		{
+		}
+
+		virtual uint32 GetID() const override { return ID; }
+		virtual FString GetName() const override { return Name; }
+		virtual uint32 GetNumButtons() const override { return Buttons.Num(); }
+
+		virtual const IStylusInputStylusButton* GetButton(int32 Index) const override
+		{
+			return 0 <= Index && Index < Buttons.Num() ? reinterpret_cast<const IStylusInputStylusButton*>(&Buttons[Index]) : nullptr;
+		}
+
+		uint32 ID = 0; // This will almost always be the physical ID, unless there is a rare collision between physical IDs of cursors of different types.
+		FString Name;
+		TArray<FStylusButton> Buttons;
+
+		uint32 WintabPhysicalID = -1;
+		uint16 WintabCursorType = -1;
+	};
+
+	using FStylusInfoContainer = TSharedRefDataContainer<FStylusInfo, false>;
+	using FStylusInfoThreadSafeContainer = TSharedRefDataContainer<FStylusInfo, false>;
+}
+
+#endif
diff --git a/Engine/Plugins/Editor/StylusInput/Source/StylusInput/Private/Windows/Wintab/WintabTabletContext.h b/Engine/Plugins/Editor/StylusInput/Source/StylusInput/Private/Windows/Wintab/WintabTabletContext.h
new file mode 100644
--- /dev/null
+++ b/Engine/Plugins/Editor/StylusInput/Source/StylusInput/Private/Windows/Wintab/WintabTabletContext.h
@@ -0,0 +1,109 @@
+﻿// Copyright Epic Games, Inc. All Rights Reserved.
+
+#pragma once
+
+#include <HAL/Platform.h>
+
+#if PLATFORM_WINDOWS && STYLUSINPUT_USE_WINTAB
+
+#include <StylusInputPacket.h>
+#include <StylusInputTabletContext.h>
+#include <StylusInputUtils.h>
+#include <Math/IntRect.h>
+#include <Math/Vector.h>
+
+#include "WintabAPI.h"
+
+namespace UE::StylusInput::Private::Wintab
+{
+	enum class EPacketPropertyType : int8
+	{
+		Status = 0,
+		Time,
+		Changed,
+		SerialNumber,
+		Cursor,
+		Buttons,
+		X,
+		Y,
+		Z,
+		NormalPressure,
+		TangentPressure,
+		Orientation,
+		Rotation,
+
+		Num_Enumerators, // THIS IS NOT A VALID ENUMERATOR, BUT IT CAN BE USED TO QUERY HOW MANY PACKET PROPERTIES THERE ARE.
+		Invalid_Enumerator = INDEX_NONE // THIS IS NOT A VALID ENUMERATOR, BUT IT CAN BE USED TO INDICATE THAT A PACKET PROPERTY ENTRY IS INVALID.
+	};
+
+	enum class ETabletPropertyMetricUnit : int8
+	{
+		Default, // Units are unknown.
+		Inches,
+		Centimeters,
+		Circle,
+
+		Num_Enumerators, // THIS IS NOT A VALID ENUMERATOR, BUT IT CAN BE USED TO QUERY HOW MANY PACKET PROPERTIES THERE ARE.
+		Invalid_Enumerator = INDEX_NONE // THIS IS NOT A VALID ENUMERATOR, BUT IT CAN BE USED TO INDICATE THAT A PACKET PROPERTY ENTRY IS INVALID.
+	};
+
+	struct FPacketProperty
+	{
+		EPacketPropertyType Type = EPacketPropertyType::Invalid_Enumerator;
+		ETabletPropertyMetricUnit MetricUnit = ETabletPropertyMetricUnit::Default;
+		int32 Minimum = 0;
+		int32 Maximum = 0;
+		float Resolution = 0.0f;
+	};
+
+	struct FPacketPropertyHandler
+	{
+		using FFuncSetProperty = void(FStylusInputPacket&, const PACKET&, const int8*);
+		FFuncSetProperty* SetProperty = nullptr;
+
+		static constexpr int32 SetPropertyDataBufferLength = 16;
+		int8 SetPropertyData[SetPropertyDataBufferLength];
+	};
+
+	class FTabletContext final : public IStylusInputTabletContext
+	{
+	public:
+		explicit FTabletContext(const uint32 ID)
+			: ID(ID)
+		{}
+
+		virtual uint32 GetID() const override { return ID; }
+
+		virtual FString GetName() const override
+		{
+			if (!Name.IsEmpty() && !PlugAndPlayID.IsEmpty())
+			{
+				return FString::Format(TEXT("{0} ({1})"), {Name, PlugAndPlayID});
+			}
+			return !Name.IsEmpty() ? Name : PlugAndPlayID;
+		}
+
+		virtual FIntRect GetInputRectangle() const override { return InputRectangle; }
+		virtual ETabletHardwareCapabilities GetHardwareCapabilities() const override { return HardwareCapabilities; }
+		virtual ETabletSupportedProperties GetSupportedProperties() const override { return SupportedProperties; }
+
+		uint32 ID = 0;
+		FString Name;
+		FString PlugAndPlayID;
+		FIntRect InputRectangle;
+		ETabletSupportedProperties SupportedProperties = ETabletSupportedProperties::None;
+		ETabletHardwareCapabilities HardwareCapabilities = ETabletHardwareCapabilities::None;
+		uint32 NumPacketProperties = 0;
+		FPacketProperty PacketProperties[static_cast<uint32>(EPacketPropertyType::Num_Enumerators)] = {};
+		FPacketPropertyHandler PacketPropertyHandlers[static_cast<uint32>(EPacketPropertyType::Num_Enumerators)] = {};
+
+		UINT WintabContextIndex = -1;
+		HCTX WintabContextHandle = nullptr;
+		uint8 WintabFirstCursor = -1;
+		uint8 WintabNumCursors = -1;
+	};
+
+	using FTabletContextContainer = TSharedRefDataContainer<FTabletContext, false>;
+}
+
+#endif
diff --git a/Engine/Plugins/Editor/StylusInput/Source/StylusInput/StylusInput.Build.cs b/Engine/Plugins/Editor/StylusInput/Source/StylusInput/StylusInput.Build.cs
--- a/Engine/Plugins/Editor/StylusInput/Source/StylusInput/StylusInput.Build.cs
+++ b/Engine/Plugins/Editor/StylusInput/Source/StylusInput/StylusInput.Build.cs
@@ -22,6 +22,17 @@ namespace UnrealBuildTool.Rules
 					"SlateCore",
 				]
 			);
+
+			const bool bUseWintab = false;
+
+#pragma warning disable CS0162
+			if (Target.Platform == UnrealTargetPlatform.Win64 && bUseWintab)
+			{
+				PrivateDependencyModuleNames.Add("Wintab");
+			}
+#pragma warning restore CS0162
+
+			PublicDefinitions.Add(bUseWintab ? "STYLUSINPUT_USE_WINTAB=1" : "STYLUSINPUT_USE_WINTAB=0");
 		}
 	}
 }
