StylusInput: Fix LOCTEXT_NAMESPACE issue

#rnx

diff --git a/Engine/Plugins/Editor/StylusInput/Source/StylusInput/Private/Windows/Wintab/WintabInterface.cpp b/Engine/Plugins/Editor/StylusInput/Source/StylusInput/Private/Windows/Wintab/WintabInterface.cpp
--- a/Engine/Plugins/Editor/StylusInput/Source/StylusInput/Private/Windows/Wintab/WintabInterface.cpp
+++ b/Engine/Plugins/Editor/StylusInput/Source/StylusInput/Private/Windows/Wintab/WintabInterface.cpp
@@ -90,4 +90,6 @@ namespace UE::StylusInput::Private::Wintab
 	}
 }
 
+#undef LOCTEXT_NAMESPACE
+
 #endif
