StylusInput: Resolve Unity build issue when Wintab is enabled

diff --git a/Engine/Plugins/Editor/StylusInput/Source/StylusInput/Private/Windows/Wintab/WintabAPI.h b/Engine/Plugins/Editor/StylusInput/Source/StylusInput/Private/Windows/Wintab/WintabAPI.h
--- a/Engine/Plugins/Editor/StylusInput/Source/StylusInput/Private/Windows/Wintab/WintabAPI.h
+++ b/Engine/Plugins/Editor/StylusInput/Source/StylusInput/Private/Windows/Wintab/WintabAPI.h
@@ -17,6 +17,7 @@
 #define PACKETDATA (/*PK_CONTEXT |*/ PK_STATUS | PK_TIME | /*PK_CHANGED |*/ PK_SERIAL_NUMBER | PK_CURSOR | PK_BUTTONS | PK_X | PK_Y | PK_Z | PK_NORMAL_PRESSURE | PK_TANGENT_PRESSURE | PK_ORIENTATION | PK_ROTATION)
 #define PACKETMODE PK_BUTTONS
 #include <PKTDEF.H>
+#undef INT // WINTAB.H defines an INT macro that collides with a macro definition in AllowWindowsPlatformTypes.h when using Unity builds.
 
 #define WTINFO_WORKAROUND_FOR_OUTPUT_BUFFER_SIZE_BUG 0
 
diff --git a/Engine/Plugins/Editor/StylusInput/Source/StylusInput/Private/Windows/Wintab/WintabInstance.cpp b/Engine/Plugins/Editor/StylusInput/Source/StylusInput/Private/Windows/Wintab/WintabInstance.cpp
--- a/Engine/Plugins/Editor/StylusInput/Source/StylusInput/Private/Windows/Wintab/WintabInstance.cpp
+++ b/Engine/Plugins/Editor/StylusInput/Source/StylusInput/Private/Windows/Wintab/WintabInstance.cpp
@@ -6,7 +6,6 @@
 
 #include <StylusInput.h>
 #include <StylusInputUtils.h>
-#include <Windows/WindowsApplication.h>
 #include <Framework/Application/SlateApplication.h>
 
 #include "WintabAPI.h"
