StylusInput: Wintab handles tablet context changes

#rnx

diff --git a/Engine/Plugins/Editor/StylusInput/Source/StylusInput/Private/Windows/Wintab/WintabAPI.h b/Engine/Plugins/Editor/StylusInput/Source/StylusInput/Private/Windows/Wintab/WintabAPI.h
--- a/Engine/Plugins/Editor/StylusInput/Source/StylusInput/Private/Windows/Wintab/WintabAPI.h
+++ b/Engine/Plugins/Editor/StylusInput/Source/StylusInput/Private/Windows/Wintab/WintabAPI.h
@@ -9,7 +9,6 @@
 #include <Containers/Array.h>
 #include <Containers/UnrealString.h>
 #include <Templates/UniquePtr.h>
-#include <Windows/WindowsPlatformProcess.h>
 
 #define WIN32_LEAN_AND_MEAN
 #include <windows.h>
@@ -154,9 +153,7 @@ namespace UE::StylusInput::Private::Wintab
 
 /*
  * TODO 
- * [ ] RealTimeStylus replace TimerTick value at earliest opportunity
- * [ ] DPI scaling?
- * [ ] Handle removing/adding tablet while running (evaluate WT_INFOCHANGE event)
+ * [X] Handle removing/adding tablet while running (evaluate WT_INFOCHANGE event)
  * [ ] PK_BUTTONS (plus RealtimeStylus buttons)
  * [ ] Queue size
  * [ ] Query multiple packets at once
@@ -165,7 +162,7 @@ namespace UE::StylusInput::Private::Wintab
  * [ ] Make available by default (with additional editor settings)
  * [ ] Propagate buttons data for both Wintab and Windows (StylusButtonDown/StylusButtonDown events)
  * [ ] Allegedly, some Wintab libraries don't handle relative button presses, so maybe switch to absolute?
- * [ ] Lost packet stats?
+ * [X] Lost packet stats?
  * [ ] RealTimeStylus StylusInRange and StylusOutOfRange events (re-evaluate stylus info?)
  */
 
diff --git a/Engine/Plugins/Editor/StylusInput/Source/StylusInput/Private/Windows/Wintab/WintabInstance.cpp b/Engine/Plugins/Editor/StylusInput/Source/StylusInput/Private/Windows/Wintab/WintabInstance.cpp
--- a/Engine/Plugins/Editor/StylusInput/Source/StylusInput/Private/Windows/Wintab/WintabInstance.cpp
+++ b/Engine/Plugins/Editor/StylusInput/Source/StylusInput/Private/Windows/Wintab/WintabInstance.cpp
@@ -545,7 +545,8 @@ namespace UE::StylusInput::Private::Wintab
 		, MessageHandler(this,
 		                 FGetTabletContextCallback::CreateRaw(this, &FWintabInstance::GetTabletContextInternal),
 		                 FGetStylusIDCallback::CreateRaw(this, &FWintabInstance::GetStylusID),
-		                 FUpdateWindowRectCallback::CreateRaw(this, &FWintabInstance::UpdateWindowRect))
+		                 FUpdateWindowRectCallback::CreateRaw(this, &FWintabInstance::UpdateWindowRect),
+						 FUpdateTabletContextsCallback::CreateRaw(this, &FWintabInstance::UpdateTabletContexts))
 	{
 		if (FSlateApplication::IsInitialized())
 		{
@@ -555,18 +556,13 @@ namespace UE::StylusInput::Private::Wintab
 			}
 		}
 
+		UpdateWindowRect();
 	    UpdateTabletContexts();
     }
 
 	FWintabInstance::~FWintabInstance()
 	{
-		for (int32 Index = 0, NumIndices = TabletContexts.Num(); Index < NumIndices; ++Index)
-		{
-			if (!WintabAPI.WtClose(TabletContexts[Index]->WintabContextHandle))
-			{
-				UE_LOG(LogStylusInput, Error, TEXT("WintabInstance: Could not close Wintab context with handle %p."), TabletContexts[Index]->WintabContextHandle);
-			}
-		}
+		ClearTabletContexts();
 
 		if (FSlateApplication::IsInitialized())
 		{
@@ -575,6 +571,20 @@ namespace UE::StylusInput::Private::Wintab
 				static_cast<FWindowsApplication*>(WindowsApplication.Get())->RemoveMessageHandler(MessageHandler);
 			}
 		}
+
+		const FPacketStats& PacketStats = MessageHandler.GetPacketStats();
+
+		if (const uint32 NumInvalidPackets = PacketStats.GetNumInvalidPackets())
+		{
+			LogVerbose(LOG_PREAMBLE, FString::Format(
+						   TEXT("Wintab instance '{0}' encountered {1} invalid packets."), {FWintabInstance::GetName().ToString(), NumInvalidPackets}));
+		}
+
+		if (const uint32 NumLostPackets = PacketStats.GetNumLostPackets())
+		{
+			LogVerbose(LOG_PREAMBLE, FString::Format(
+						   TEXT("Wintab instance '{0}' encountered {1} lost packets."), {FWintabInstance::GetName().ToString(), NumLostPackets}));
+		}
 	}
 
 	bool FWintabInstance::AddEventHandler(IStylusInputEventHandler* EventHandler, const EEventHandlerThread Thread)
@@ -629,7 +639,7 @@ namespace UE::StylusInput::Private::Wintab
 
 	float FWintabInstance::GetPacketsPerSecond(const EEventHandlerThread EventHandlerThread) const
 	{
-        return MessageHandler.GetPacketsPerSecond();
+        return MessageHandler.GetPacketStats().GetPacketsPerSecond();
     }
 
 	FText FWintabInstance::GetName()
@@ -637,9 +647,9 @@ namespace UE::StylusInput::Private::Wintab
 		return FText::Format(LOCTEXT("Wintab", "Wintab #{0}"), ID);
 	}
 
-	const FTabletContext& FWintabInstance::GetTabletContextInternal(const uint32 TabletContextID) const
+	const FTabletContext* FWintabInstance::GetTabletContextInternal(const uint32 TabletContextID) const
 	{
-		return *TabletContexts.Get(TabletContextID);
+		return TabletContexts.Get(TabletContextID).Get();
 	}
 
 	uint32 FWintabInstance::GetStylusID(const uint32 TabletContextID, uint32 CursorIndex)
@@ -766,9 +776,30 @@ namespace UE::StylusInput::Private::Wintab
 		return 0;
 	}
 
+	void FWintabInstance::ClearTabletContexts()
+	{
+		for (int32 Index = 0, NumIndices = TabletContexts.Num(); Index < NumIndices; ++Index)
+		{
+			const HCTX WintabContextHandle = TabletContexts[Index]->WintabContextHandle;
+
+			if (WintabAPI.WtClose(WintabContextHandle))
+			{
+				LogVerbose(LOG_PREAMBLE, FString::Format(TEXT("Closed Wintab context with handle {0}."),
+														 {HctxToUint32(WintabContextHandle)}));
+			}
+			else
+			{
+				LogError(LOG_PREAMBLE, FString::Format(TEXT("Could not close Wintab context with handle {0}."),
+													   {HctxToUint32(WintabContextHandle)}));
+			}
+		}
+
+		TabletContexts.Clear();
+	}
+
 	void FWintabInstance::UpdateTabletContexts()
 	{
-		UpdateWindowRect();
+		ClearTabletContexts();
 
 		UINT WintabContextIndex = 0;
 		LOGCONTEXTW WintabContext;
@@ -815,7 +846,7 @@ namespace UE::StylusInput::Private::Wintab
 				continue;
 			}
 
-			LogVerbose(LOG_PREAMBLE, FString::Format(TEXT("Added TabletContext for ID {0} [{1}, {2}]."), {
+			LogVerbose(LOG_PREAMBLE, FString::Format(TEXT("Added tablet context for ID {0} [{1}, {2}]."), {
 				                                         TabletContext.ID, TabletContext.Name , TabletContext.PlugAndPlayID
 			                                         }));
 
diff --git a/Engine/Plugins/Editor/StylusInput/Source/StylusInput/Private/Windows/Wintab/WintabInstance.h b/Engine/Plugins/Editor/StylusInput/Source/StylusInput/Private/Windows/Wintab/WintabInstance.h
--- a/Engine/Plugins/Editor/StylusInput/Source/StylusInput/Private/Windows/Wintab/WintabInstance.h
+++ b/Engine/Plugins/Editor/StylusInput/Source/StylusInput/Private/Windows/Wintab/WintabInstance.h
@@ -41,8 +41,10 @@ namespace UE::StylusInput::Private::Wintab
 		virtual FText GetName() override;
 
 	private:
-		const FTabletContext& GetTabletContextInternal(uint32 TabletContextID) const;
+		const FTabletContext* GetTabletContextInternal(uint32 TabletContextID) const;
 		uint32 GetStylusID(uint32 TabletContextID, uint32 CursorIndex);
+
+		void ClearTabletContexts();
 		void UpdateTabletContexts();
 		void UpdateWindowRect();
 
diff --git a/Engine/Plugins/Editor/StylusInput/Source/StylusInput/Private/Windows/Wintab/WintabMessageHandler.cpp b/Engine/Plugins/Editor/StylusInput/Source/StylusInput/Private/Windows/Wintab/WintabMessageHandler.cpp
--- a/Engine/Plugins/Editor/StylusInput/Source/StylusInput/Private/Windows/Wintab/WintabMessageHandler.cpp
+++ b/Engine/Plugins/Editor/StylusInput/Source/StylusInput/Private/Windows/Wintab/WintabMessageHandler.cpp
@@ -18,17 +18,24 @@
 #define ENABLE_DEBUG_EVENTS_FOR_PROXIMITY_CSRCHANGE_MESSAGES 0
 #define ENABLE_DEBUG_EVENTS_FOR_INFOCHANGE_MESSAGES 0
 #define ENABLE_DEBUG_EVENTS_FOR_PACKETEXT_MESSAGES 0
+#define ENABLE_DEBUG_EVENTS_FOR_INVALID_PACKETS 0
 
 namespace UE::StylusInput::Private::Wintab
 {
 FWintabMessageHandler::FWintabMessageHandler(IStylusInputInstance* Instance, FGetTabletContextCallback&& GetTabletContextCallback,
-                                             FGetStylusIDCallback&& GetStylusIDCallback, FUpdateWindowRectCallback&& UpdateWindowRectCallback)
+                                             FGetStylusIDCallback&& GetStylusIDCallback, FUpdateWindowRectCallback&& UpdateWindowRectCallback,
+                                             FUpdateTabletContextsCallback&& UpdateTabletContextsCallback)
 	: Instance(Instance)
 	, WintabAPI(FWintabAPI::GetInstance())
 	, GetTabletContextCallback(MoveTemp(GetTabletContextCallback))
 	, GetStylusIDCallback(MoveTemp(GetStylusIDCallback))
 	, UpdateWindowRectCallback(MoveTemp(UpdateWindowRectCallback))
+	, UpdateTabletContextsCallback(MoveTemp(UpdateTabletContextsCallback))
 {
+	check(this->GetTabletContextCallback.IsBound())
+	check(this->GetStylusIDCallback.IsBound())
+	check(this->UpdateWindowRectCallback.IsBound())
+	check(this->UpdateTabletContextsCallback.IsBound())
 }
 
 FWintabMessageHandler::~FWintabMessageHandler()
@@ -52,21 +59,31 @@ bool FWintabMessageHandler::ProcessMessage(const HWND Hwnd, const uint32 Msg, co
 				PACKET WintabPacket = {};
 				if (WintabAPI.WtPacket(reinterpret_cast<HCTX>(LParam), WParam, &WintabPacket))
 				{
+					if (bCursorChange)
+					{
+						ProcessCursorChange(static_cast<uint32>(LParam), &WintabPacket);
+					}
+
 					// Todo Process more than one packet? 
 
 					PacketStats.NewPacket(WintabPacket.pkSerialNumber);
 
 					const FStylusInputPacket Packet = ProcessPacket(static_cast<uint32>(LParam), &WintabPacket);
 
-					for (IStylusInputEventHandler* EventHandler : EventHandlers)
+					if (Packet.Type != EPacketType::Invalid)
 					{
-						EventHandler->OnPacket(Packet, Instance);
+						for (IStylusInputEventHandler* EventHandler : EventHandlers)
+						{
+							EventHandler->OnPacket(Packet, Instance);
+						}
 					}
 				}
 				else
 				{
 					PacketStats.InvalidPacket();
+#if ENABLE_DEBUG_EVENTS_FOR_INVALID_PACKETS
 					DebugEvent(FString::Format(TEXT("Failed to retrieve Wintab packet number {0}."), {WParam}));
+#endif
 				}
 			}
 
@@ -76,7 +93,7 @@ bool FWintabMessageHandler::ProcessMessage(const HWND Hwnd, const uint32 Msg, co
 	case WT_CTXOPEN:
 		{
 #if ENABLE_DEBUG_EVENTS_FOR_CONTEXT_OPEN_CLOSE_UPDATE_MESSAGES
-			DebugEvent("Received WT_CTXOPEN message.");
+			DebugEvent(FString::Format(TEXT("Received WT_CTXOPEN message (Context={0}, Status={1})."), {WParam, LParam}));
 #endif
 
 			bHandled = true;
@@ -85,7 +102,7 @@ bool FWintabMessageHandler::ProcessMessage(const HWND Hwnd, const uint32 Msg, co
 	case WT_CTXCLOSE:
 		{
 #if ENABLE_DEBUG_EVENTS_FOR_CONTEXT_OPEN_CLOSE_UPDATE_MESSAGES
-			DebugEvent("Received WT_CTXCLOSE message.");
+			DebugEvent(FString::Format(TEXT("Received WT_CTXCLOSE message (Context={0}, Status={1})."), {WParam, LParam}));
 #endif
 
 			bHandled = true;
@@ -94,7 +111,7 @@ bool FWintabMessageHandler::ProcessMessage(const HWND Hwnd, const uint32 Msg, co
 	case WT_CTXUPDATE:
 		{
 #if ENABLE_DEBUG_EVENTS_FOR_CONTEXT_OPEN_CLOSE_UPDATE_MESSAGES
-			DebugEvent("Received WT_CTXUPDATE message.");
+			DebugEvent(FString::Format(TEXT("Received WT_CTXUPDATE message (Context={0}, Status={1})."), {WParam, LParam}));
 #endif
 
 			bHandled = true;
@@ -103,7 +120,7 @@ bool FWintabMessageHandler::ProcessMessage(const HWND Hwnd, const uint32 Msg, co
 	case WT_CTXOVERLAP:
 		{
 #if ENABLE_DEBUG_EVENTS_FOR_CONTEXT_OVERLAP_MESSAGES
-			DebugEvent("Received WT_CTXOVERLAP message.");
+			DebugEvent(FString::Format(TEXT("Received WT_CTXOVERLAP message (Context={0}, Status={1})."), {WParam, LParam}));
 #endif
 
 			bHandled = true;
@@ -136,6 +153,8 @@ bool FWintabMessageHandler::ProcessMessage(const HWND Hwnd, const uint32 Msg, co
 			DebugEvent(FString::Format(TEXT("Received WT_INFOCHANGE(Context={0}, Category={1}, Index={2}) message."), {WParam, LOWORD(LParam), HIWORD(LParam)}));
 #endif
 
+			UpdateTabletContextsCallback.Execute();
+
 			bHandled = true;
 			break;
 		}
@@ -212,12 +231,6 @@ void FWintabMessageHandler::DebugEvent(const FString& Message) const
 	}
 }
 
-const FPacketProperty* FWintabMessageHandler::GetPacketDescriptions(uint32 TabletContextID) const
-{
-	const FTabletContext& TabletContext = GetTabletContextCallback.Execute(TabletContextID);
-	return /*TabletContext ?*/ TabletContext.PacketProperties /*: nullptr*/;
-}
-
 void FWintabMessageHandler::ProcessCursorChange(const uint32 TabletContextID, const void *const WintabPacketPtr)
 {
 	const PACKET& WintabPacket = *static_cast<const PACKET*>(WintabPacketPtr);
@@ -243,27 +256,23 @@ void FWintabMessageHandler::ProcessCursorChange(const uint32 TabletContextID, co
 
 FStylusInputPacket FWintabMessageHandler::ProcessPacket(const uint32 TabletContextID, const void *const WintabPacketPtr)
 {
-	const PACKET& WintabPacket = *static_cast<const PACKET*>(WintabPacketPtr);
+	FStylusInputPacket Packet{TabletContextID};
 
-	if (bCursorChange)
+	if (const FTabletContext* TabletContext = GetTabletContextCallback.Execute(TabletContextID))
 	{
-		ProcessCursorChange(TabletContextID, WintabPacketPtr);
-	}
-
-	const FTabletContext& TabletContext = GetTabletContextCallback.Execute(TabletContextID);
-
-	FStylusInputPacket Packet{TabletContext.ID};
+		const PACKET& WintabPacket = *static_cast<const PACKET*>(WintabPacketPtr);
 
-	const bool bPacketOnDigitizer = !static_cast<bool>(TabletContext.SupportedProperties & ETabletSupportedProperties::Z) || WintabPacket.pkZ == 0;
-	Packet.Type = bPacketOnDigitizer
-					  ? bLastPacketOnDigitizer ? EPacketType::OnDigitizer : EPacketType::StylusDown
-					  : bLastPacketOnDigitizer ? EPacketType::StylusUp : EPacketType::AboveDigitizer;
-	bLastPacketOnDigitizer = bPacketOnDigitizer;
+		const bool bPacketOnDigitizer = !static_cast<bool>(TabletContext->SupportedProperties & ETabletSupportedProperties::Z) || WintabPacket.pkZ == 0;
+		Packet.Type = bPacketOnDigitizer
+						  ? bLastPacketOnDigitizer ? EPacketType::OnDigitizer : EPacketType::StylusDown
+						  : bLastPacketOnDigitizer ? EPacketType::StylusUp : EPacketType::AboveDigitizer;
+		bLastPacketOnDigitizer = bPacketOnDigitizer;
 
-	for (int32 Index = 0, Num = TabletContext.NumPacketProperties; Index < Num; ++Index)
-	{
-		const FPacketPropertyHandler& PropertyHandler = TabletContext.PacketPropertyHandlers[Index];
-		PropertyHandler.SetProperty(Packet, WintabPacket, PropertyHandler.SetPropertyData);
+		for (int32 Index = 0, Num = TabletContext->NumPacketProperties; Index < Num; ++Index)
+		{
+			const FPacketPropertyHandler& PropertyHandler = TabletContext->PacketPropertyHandlers[Index];
+			PropertyHandler.SetProperty(Packet, WintabPacket, PropertyHandler.SetPropertyData);
+		}
 	}
 
 	return Packet;
diff --git a/Engine/Plugins/Editor/StylusInput/Source/StylusInput/Private/Windows/Wintab/WintabMessageHandler.h b/Engine/Plugins/Editor/StylusInput/Source/StylusInput/Private/Windows/Wintab/WintabMessageHandler.h
--- a/Engine/Plugins/Editor/StylusInput/Source/StylusInput/Private/Windows/Wintab/WintabMessageHandler.h
+++ b/Engine/Plugins/Editor/StylusInput/Source/StylusInput/Private/Windows/Wintab/WintabMessageHandler.h
@@ -22,15 +22,16 @@ namespace UE::StylusInput::Private::Wintab
 	class FWintabAPI;
 
 	DECLARE_DELEGATE_RetVal(const FWindowContext&, FGetWindowContextCallback);
-	DECLARE_DELEGATE_RetVal_OneParam(const FTabletContext&, FGetTabletContextCallback, uint32);
+	DECLARE_DELEGATE_RetVal_OneParam(const FTabletContext*, FGetTabletContextCallback, uint32);
 	DECLARE_DELEGATE_RetVal_TwoParams(uint32, FGetStylusIDCallback, uint32, UINT);
 	DECLARE_DELEGATE(FUpdateWindowRectCallback);
+	DECLARE_DELEGATE(FUpdateTabletContextsCallback);
 
 	class FWintabMessageHandler final : public IWindowsMessageHandler
 	{
 	public:
 		FWintabMessageHandler(IStylusInputInstance* Instance, FGetTabletContextCallback&& GetTabletContextCallback, FGetStylusIDCallback&& GetStylusIDCallback,
-		                      FUpdateWindowRectCallback&& UpdateWindowRectCallback);
+		                      FUpdateWindowRectCallback&& UpdateWindowRectCallback, FUpdateTabletContextsCallback&& UpdateTabletContextsCallback);
 		virtual ~FWintabMessageHandler();
 
 		virtual bool ProcessMessage(HWND Hwnd, uint32 Msg, WPARAM WParam, LPARAM LParam, int32& OutResult) override;
@@ -39,11 +40,10 @@ namespace UE::StylusInput::Private::Wintab
 		bool RemoveEventHandler(IStylusInputEventHandler* EventHandler);
 
 		uint32 GetCurrentStylusID() const { return CurrentStylusID; }
-		float GetPacketsPerSecond() const { return PacketStats.GetPacketsPerSecond(); }
+		const FPacketStats& GetPacketStats() const { return PacketStats; }
 
 	private:
 		void DebugEvent(const FString& Message) const;
-		const FPacketProperty* GetPacketDescriptions(uint32 TabletContextID) const;
 
 		void ProcessCursorChange(uint32 TabletContextID, const void* WintabPacketPtr);
 		FStylusInputPacket ProcessPacket(uint32 TabletContextID, const void* WintabPacketPtr);
@@ -53,6 +53,7 @@ namespace UE::StylusInput::Private::Wintab
 		FGetTabletContextCallback GetTabletContextCallback;
 		FGetStylusIDCallback GetStylusIDCallback;
 		FUpdateWindowRectCallback UpdateWindowRectCallback;
+		FUpdateTabletContextsCallback UpdateTabletContextsCallback;
 		TArray<IStylusInputEventHandler*> EventHandlers;
 		FPacketStats PacketStats;
 
diff --git a/Engine/Plugins/Editor/StylusInput/Source/StylusInput/Private/Windows/Wintab/WintabStats.h b/Engine/Plugins/Editor/StylusInput/Source/StylusInput/Private/Windows/Wintab/WintabStats.h
--- a/Engine/Plugins/Editor/StylusInput/Source/StylusInput/Private/Windows/Wintab/WintabStats.h
+++ b/Engine/Plugins/Editor/StylusInput/Source/StylusInput/Private/Windows/Wintab/WintabStats.h
@@ -36,23 +36,18 @@ namespace UE::StylusInput::Private::Wintab
 			++NumInvalidPackets;
 		}
 
-		void LostPackets()
-		{
-			++NumLostPackets;
-		}
-
 		float GetPacketsPerSecond() const
 		{
 			const uint32 Timestamp = FPlatformTime::Cycles();
 			return Timestamp - LatestTimestamp <= CyclesPerSecond ? PacketsPerSecond : 0.0f;
 		}
 
-		uint32 GetInvalidPackets() const
+		uint32 GetNumInvalidPackets() const
 		{
 			return NumInvalidPackets;
 		}
 
-		uint32 GetLostPackets() const
+		uint32 GetNumLostPackets() const
 		{
 			return NumLostPackets;
 		}
