StylusInput: RealTimeStylus: Update cursor for each StylusInRange event

#rnx

diff --git a/Engine/Plugins/Editor/StylusInput/Source/StylusInput/Private/StylusInputUtils.h b/Engine/Plugins/Editor/StylusInput/Source/StylusInput/Private/StylusInputUtils.h
--- a/Engine/Plugins/Editor/StylusInput/Source/StylusInput/Private/StylusInputUtils.h
+++ b/Engine/Plugins/Editor/StylusInput/Source/StylusInput/Private/StylusInputUtils.h
@@ -35,24 +35,45 @@ namespace UE::StylusInput::Private
 			return Data.Emplace_GetRef(MakeShared<DataType>(ID));
 		}
 
+		void AddOrReplace(TSharedPtr<DataType>&& Item)
+		{
+			if (Item.IsValid())
+			{
+				TSharedRef<DataType>* ExistingItem = Data.FindByPredicate([ID = Item->ID](const TSharedRef<DataType>& Context)
+				{
+					return Context->ID == ID;
+				});
+
+				WriteScopeLockType WriteLock(RWLock);
+				if (ExistingItem != nullptr)
+				{
+					*ExistingItem = MoveTemp(Item).ToSharedRef();
+				}
+				else
+				{
+					Data.Emplace(MoveTemp(Item).ToSharedRef());
+				}
+			}
+		}
+
 		bool Contains(uint32 ID) const
 		{
 			ReadScopeLockType ReadLock(RWLock);
-			const TSharedRef<DataType>* TabletContext = Data.FindByPredicate([ID](const TSharedRef<DataType>& Context)
+			const TSharedRef<DataType>* Item = Data.FindByPredicate([ID](const TSharedRef<DataType>& Context)
 			{
 				return Context->ID == ID;
 			});
-			return TabletContext != nullptr;
+			return Item != nullptr;
 		}
 
 		TSharedPtr<DataType> Get(uint32 ID) const
 		{
 			ReadScopeLockType ReadLock(RWLock);
-			const TSharedRef<DataType>* TabletContext = Data.FindByPredicate([ID](const TSharedRef<DataType>& Context)
+			const TSharedRef<DataType>* Item = Data.FindByPredicate([ID](const TSharedRef<DataType>& Context)
 			{
 				return Context->ID == ID;
 			});
-			return TabletContext ? TabletContext->ToSharedPtr() : TSharedPtr<DataType>{};
+			return Item ? Item->ToSharedPtr() : TSharedPtr<DataType>{};
 		}
 
 		bool Remove(uint32 ID)
@@ -77,13 +98,13 @@ namespace UE::StylusInput::Private
 			Data.Reset();
 		}
 
-		void Update(const TSharedRefDataContainer<DataType, false>& InTabletContexts)
+		void Update(const TSharedRefDataContainer<DataType, false>& InData)
 		{
 			WriteScopeLockType WriteLock(RWLock);
 			Data.Reset();
-			for (uint32 Index = 0, Num = InTabletContexts.Num(); Index < Num; ++Index)
+			for (uint32 Index = 0, Num = InData.Num(); Index < Num; ++Index)
 			{
-				Data.Emplace(InTabletContexts[Index]);
+				Data.Emplace(InData[Index]);
 			}
 		}
 
diff --git a/Engine/Plugins/Editor/StylusInput/Source/StylusInput/Private/Windows/RealTimeStylus/RealTimeStylusInstance.cpp b/Engine/Plugins/Editor/StylusInput/Source/StylusInput/Private/Windows/RealTimeStylus/RealTimeStylusInstance.cpp
--- a/Engine/Plugins/Editor/StylusInput/Source/StylusInput/Private/Windows/RealTimeStylus/RealTimeStylusInstance.cpp
+++ b/Engine/Plugins/Editor/StylusInput/Source/StylusInput/Private/Windows/RealTimeStylus/RealTimeStylusInstance.cpp
@@ -132,102 +132,7 @@
 
 	const TSharedPtr<IStylusInputStylusInfo> FRealTimeStylusInstance::GetStylusInfo(uint32 StylusID)
 	{
-		TSharedPtr<FStylusInfo> StylusInfo = StylusInfos.Get(StylusID);
-		if (StylusInfo.IsValid())
-		{
-			return StylusInfo;
-		}
-
-		StylusInfo = StylusInfos.Add(StylusID);
-
-		IInkCursor* InkCursor = nullptr;
-		if (Succeeded(RealTimeStylus->GetStylusForId(StylusID, &InkCursor), LOG_PREAMBLE) && InkCursor != nullptr)
-		{
-			BSTR String = nullptr;
-			if (Succeeded(InkCursor->get_Name(&String), LOG_PREAMBLE))
-			{
-				StylusInfo->Name = String;
-				RealTimeStylusAPI.SysFreeString(String);
-			}
-			else
-			{
-				LogError(LOG_PREAMBLE, FString::Format(TEXT("Could not get name for stylus with ID {0}."), {StylusID}));
-			}
-
-			IInkCursorButtons* Buttons = nullptr;
-			long NumButtons = 0;
-			if (Succeeded(InkCursor->get_Buttons(&Buttons), LOG_PREAMBLE) &&
-				Buttons && 
-				Succeeded(Buttons->get_Count(&NumButtons), LOG_PREAMBLE))
-			{
-				if (NumButtons > 0)
-				{
-					TArray<FStylusButton> StylusInfoButtons;
-					StylusInfoButtons.Reserve(NumButtons);
-
-					for (long ButtonIndex = 0; ButtonIndex < NumButtons; ++ButtonIndex)
-					{
-						VARIANT ButtonIdentifier;
-						RealTimeStylusAPI.VariantInit(&ButtonIdentifier);
-
-						ButtonIdentifier.vt = VT_INT;
-						ButtonIdentifier.lVal = ButtonIndex;
-
-						IInkCursorButton* Button = nullptr;
-						if (Succeeded(Buttons->Item(ButtonIdentifier, &Button), LOG_PREAMBLE))
-						{
-							FStylusButton NewButton;
-
-							if (Succeeded(Button->get_Id(&String), LOG_PREAMBLE))
-							{
-								NewButton.ID = String;
-								RealTimeStylusAPI.SysFreeString(String);
-							}
-							else
-							{
-								LogError(LOG_PREAMBLE, FString::Format(
-											 TEXT("Could not get ID for button {0} for stylus with ID {1}."), {static_cast<int32>(ButtonIndex), StylusID}));
-							}
-
-							if (Succeeded(Button->get_Name(&String), LOG_PREAMBLE))
-							{
-								NewButton.Name = String;
-								RealTimeStylusAPI.SysFreeString(String);
-							}
-							else
-							{
-								LogError(LOG_PREAMBLE, FString::Format(
-											 TEXT("Could not get name for button {0} for stylus with ID {1}."), {static_cast<int32>(ButtonIndex), StylusID}));
-							}
-
-							if (!NewButton.ID.IsEmpty() && !NewButton.Name.IsEmpty())
-							{
-								StylusInfoButtons.Emplace(NewButton);
-							}
-						}
-						else
-						{
-							LogError(LOG_PREAMBLE, FString::Format(
-										 TEXT("Could not get button {0} for stylus with ID {1}."), {static_cast<int32>(ButtonIndex), StylusID}));
-						}
-
-						RealTimeStylusAPI.VariantClear(&ButtonIdentifier);
-					}
-
-					StylusInfo->Buttons = MoveTemp(StylusInfoButtons);
-				}
-			}
-			else
-			{
-				LogError(LOG_PREAMBLE, FString::Format(TEXT("Could not get buttons for stylus with ID {0}."), {StylusID}));
-			}
-		}
-		else
-		{
-			LogWarning(LOG_PREAMBLE, FString::Format(TEXT("Could not get stylus info for ID {0}."), {StylusID}));
-		}
-
-		return StylusInfo;
+		return StylusInfos.Get(StylusID);
 	}
 
 	float FRealTimeStylusInstance::GetPacketsPerSecond(const EEventHandlerThread EventHandlerThread) const
@@ -329,6 +234,101 @@
 		return WindowContext;
 	}
 
+	void FRealTimeStylusInstance::UpdateStylusInfo(STYLUS_ID StylusID)
+	{
+		LogWarning(LOG_PREAMBLE, "FRealTimeStylusInstance::UpdateStylusInfo()");
+
+		TSharedPtr<FStylusInfo> StylusInfo = MakeShared<FStylusInfo>(StylusID);
+
+		IInkCursor* InkCursor = nullptr;
+		if (Succeeded(RealTimeStylus->GetStylusForId(StylusID, &InkCursor), LOG_PREAMBLE) && InkCursor != nullptr)
+		{
+			BSTR String = nullptr;
+			if (Succeeded(InkCursor->get_Name(&String), LOG_PREAMBLE))
+			{
+				StylusInfo->Name = String;
+				RealTimeStylusAPI.SysFreeString(String);
+			}
+			else
+			{
+				LogError(LOG_PREAMBLE, FString::Format(TEXT("Could not get name for stylus with ID {0}."), {static_cast<uint32>(StylusID)}));
+			}
+
+			IInkCursorButtons* Buttons = nullptr;
+			long NumButtons = 0;
+			if (Succeeded(InkCursor->get_Buttons(&Buttons), LOG_PREAMBLE) &&
+				Buttons && Succeeded(Buttons->get_Count(&NumButtons), LOG_PREAMBLE))
+			{
+				if (NumButtons > 0)
+				{
+					TArray<FStylusButton> StylusInfoButtons;
+					StylusInfoButtons.Reserve(NumButtons);
+
+					for (long ButtonIndex = 0; ButtonIndex < NumButtons; ++ButtonIndex)
+					{
+						VARIANT ButtonIdentifier;
+						RealTimeStylusAPI.VariantInit(&ButtonIdentifier);
+
+						ButtonIdentifier.vt = VT_INT;
+						ButtonIdentifier.lVal = ButtonIndex;
+
+						IInkCursorButton* Button = nullptr;
+						if (Succeeded(Buttons->Item(ButtonIdentifier, &Button), LOG_PREAMBLE))
+						{
+							FStylusButton NewButton;
+
+							if (Succeeded(Button->get_Id(&String), LOG_PREAMBLE))
+							{
+								NewButton.ID = String;
+								RealTimeStylusAPI.SysFreeString(String);
+							}
+							else
+							{
+								LogError(LOG_PREAMBLE, FString::Format(
+											 TEXT("Could not get ID for button {0} for stylus with ID {1}."), {static_cast<int32>(ButtonIndex), static_cast<uint32>(StylusID)}));
+							}
+
+							if (Succeeded(Button->get_Name(&String), LOG_PREAMBLE))
+							{
+								NewButton.Name = String;
+								RealTimeStylusAPI.SysFreeString(String);
+							}
+							else
+							{
+								LogError(LOG_PREAMBLE, FString::Format(
+											 TEXT("Could not get name for button {0} for stylus with ID {1}."), {static_cast<int32>(ButtonIndex), static_cast<uint32>(StylusID)}));
+							}
+
+							if (!NewButton.ID.IsEmpty() && !NewButton.Name.IsEmpty())
+							{
+								StylusInfoButtons.Emplace(NewButton);
+							}
+						}
+						else
+						{
+							LogError(LOG_PREAMBLE, FString::Format(
+										 TEXT("Could not get button {0} for stylus with ID {1}."), {static_cast<int32>(ButtonIndex), static_cast<uint32>(StylusID)}));
+						}
+
+						RealTimeStylusAPI.VariantClear(&ButtonIdentifier);
+					}
+
+					StylusInfo->Buttons = MoveTemp(StylusInfoButtons);
+				}
+			}
+			else
+			{
+				LogError(LOG_PREAMBLE, FString::Format(TEXT("Could not get buttons for stylus with ID {0}."), {static_cast<uint32>(StylusID)}));
+			}
+		}
+		else
+		{
+			LogWarning(LOG_PREAMBLE, FString::Format(TEXT("Could not get stylus info for ID {0}."), {static_cast<uint32>(StylusID)}));
+		}
+
+		StylusInfos.AddOrReplace(MoveTemp(StylusInfo));
+	}
+
 	void FRealTimeStylusInstance::UpdateTabletContexts(const FTabletContextContainer& InTabletContexts)
 	{
 		TabletContexts.Update(InTabletContexts);
@@ -382,6 +382,7 @@
 				this,
 				FGetWindowContextCallback::CreateRaw(this, &FRealTimeStylusInstance::GetWindowContext),
 				FUpdateTabletContextsCallback::CreateRaw(this, &FRealTimeStylusInstance::UpdateTabletContexts),
+				FUpdateStylusInfoCallback::CreateRaw(this, &FRealTimeStylusInstance::UpdateStylusInfo),
 				EventHandler
 			);
 
@@ -402,6 +403,7 @@
 				this,
 				FGetWindowContextCallback::CreateRaw(this, &FRealTimeStylusInstance::GetWindowContext),
 				FUpdateTabletContextsCallback::CreateRaw(this, &FRealTimeStylusInstance::UpdateTabletContexts),
+				FUpdateStylusInfoCallback::CreateRaw(this, &FRealTimeStylusInstance::UpdateStylusInfo),
 				EventHandler
 			);
 
diff --git a/Engine/Plugins/Editor/StylusInput/Source/StylusInput/Private/Windows/RealTimeStylus/RealTimeStylusInstance.h b/Engine/Plugins/Editor/StylusInput/Source/StylusInput/Private/Windows/RealTimeStylus/RealTimeStylusInstance.h
--- a/Engine/Plugins/Editor/StylusInput/Source/StylusInput/Private/Windows/RealTimeStylus/RealTimeStylusInstance.h
+++ b/Engine/Plugins/Editor/StylusInput/Source/StylusInput/Private/Windows/RealTimeStylus/RealTimeStylusInstance.h
@@ -53,6 +53,7 @@ namespace UE::StylusInput::Private::RealTimeStylus
 		void SetupWindowContext(HWND HWindow);
 		const FWindowContext& GetWindowContext() const;
 
+		void UpdateStylusInfo(STYLUS_ID StylusID);
 		void UpdateTabletContexts(const FTabletContextContainer& InTabletContexts);
 
 		const uint32 ID;
diff --git a/Engine/Plugins/Editor/StylusInput/Source/StylusInput/Private/Windows/RealTimeStylus/RealTimeStylusPluginAsync.cpp b/Engine/Plugins/Editor/StylusInput/Source/StylusInput/Private/Windows/RealTimeStylus/RealTimeStylusPluginAsync.cpp
--- a/Engine/Plugins/Editor/StylusInput/Source/StylusInput/Private/Windows/RealTimeStylus/RealTimeStylusPluginAsync.cpp
+++ b/Engine/Plugins/Editor/StylusInput/Source/StylusInput/Private/Windows/RealTimeStylus/RealTimeStylusPluginAsync.cpp
@@ -9,8 +9,9 @@ namespace UE::StylusInput::Private::RealTimeStylus
 	FRealTimeStylusPluginAsync::FRealTimeStylusPluginAsync(IStylusInputInstance* Instance,
 	                                                       FGetWindowContextCallback&& GetWindowContext,
 	                                                       FUpdateTabletContextsCallback&& UpdateTabletContextsCallback,
+	                                                       FUpdateStylusInfoCallback UpdateStylusInfoCallback,
 	                                                       IStylusInputEventHandler* EventHandler)
-		: FRealTimeStylusPluginBase(Instance, MoveTemp(GetWindowContext), MoveTemp(UpdateTabletContextsCallback))
+		: FRealTimeStylusPluginBase(Instance, MoveTemp(GetWindowContext), MoveTemp(UpdateTabletContextsCallback), MoveTemp(UpdateStylusInfoCallback))
 	{
 		if (EventHandler)
 		{
@@ -71,7 +72,7 @@ namespace UE::StylusInput::Private::RealTimeStylus
 
 	HRESULT FRealTimeStylusPluginAsync::StylusInRange(IRealTimeStylus* RealTimeStylus, TABLET_CONTEXT_ID TabletContextID, STYLUS_ID StylusID)
 	{
-		return E_NOTIMPL;
+		return ProcessStylusInRange(RealTimeStylus, TabletContextID, StylusID);
 	}
 
 	HRESULT FRealTimeStylusPluginAsync::StylusOutOfRange(IRealTimeStylus* RealTimeStylus, TABLET_CONTEXT_ID TabletContextID, STYLUS_ID StylusID)
diff --git a/Engine/Plugins/Editor/StylusInput/Source/StylusInput/Private/Windows/RealTimeStylus/RealTimeStylusPluginAsync.h b/Engine/Plugins/Editor/StylusInput/Source/StylusInput/Private/Windows/RealTimeStylus/RealTimeStylusPluginAsync.h
--- a/Engine/Plugins/Editor/StylusInput/Source/StylusInput/Private/Windows/RealTimeStylus/RealTimeStylusPluginAsync.h
+++ b/Engine/Plugins/Editor/StylusInput/Source/StylusInput/Private/Windows/RealTimeStylus/RealTimeStylusPluginAsync.h
@@ -14,7 +14,8 @@ namespace UE::StylusInput::Private::RealTimeStylus
 	{
 	public:
 		FRealTimeStylusPluginAsync(IStylusInputInstance* Instance, FGetWindowContextCallback&& GetWindowContext,
-		                               FUpdateTabletContextsCallback&& UpdateTabletContextsCallback, IStylusInputEventHandler* EventHandler);
+		                           FUpdateTabletContextsCallback&& UpdateTabletContextsCallback, FUpdateStylusInfoCallback UpdateStylusInfoCallback,
+		                           IStylusInputEventHandler* EventHandler);
 		~FRealTimeStylusPluginAsync();
 
 		// IUnknown
diff --git a/Engine/Plugins/Editor/StylusInput/Source/StylusInput/Private/Windows/RealTimeStylus/RealTimeStylusPluginBase.cpp b/Engine/Plugins/Editor/StylusInput/Source/StylusInput/Private/Windows/RealTimeStylus/RealTimeStylusPluginBase.cpp
--- a/Engine/Plugins/Editor/StylusInput/Source/StylusInput/Private/Windows/RealTimeStylus/RealTimeStylusPluginBase.cpp
+++ b/Engine/Plugins/Editor/StylusInput/Source/StylusInput/Private/Windows/RealTimeStylus/RealTimeStylusPluginBase.cpp
@@ -522,11 +522,13 @@ namespace UE::StylusInput::Private::RealTimeStylus
 	}
 
 	FRealTimeStylusPluginBase::FRealTimeStylusPluginBase(IStylusInputInstance* Instance,
-	                                                             FGetWindowContextCallback&& GetWindowContextCallback,
-	                                                             FUpdateTabletContextsCallback&& UpdateTabletContextsCallback)
+	                                                     FGetWindowContextCallback&& GetWindowContextCallback,
+	                                                     FUpdateTabletContextsCallback&& UpdateTabletContextsCallback,
+	                                                     FUpdateStylusInfoCallback&& UpdateStylusInfoCallback)
 		: Instance(Instance)
 		, GetWindowContextCallback(MoveTemp(GetWindowContextCallback))
-		, UpdateTabletContextsCallback(MoveTemp(UpdateTabletContextsCallback)) 
+		, UpdateTabletContextsCallback(MoveTemp(UpdateTabletContextsCallback))
+		, UpdateStylusInfoCallback(MoveTemp(UpdateStylusInfoCallback))
 	{
 		check(this->GetWindowContextCallback.IsBound())
 		check(this->UpdateTabletContextsCallback.IsBound())
@@ -556,10 +558,10 @@ namespace UE::StylusInput::Private::RealTimeStylus
 				RTSDI_Error |
 				RTSDI_RealTimeStylusEnabled | RTSDI_RealTimeStylusDisabled |
 				RTSDI_InAirPackets | RTSDI_Packets |
-				RTSDI_StylusDown | RTSDI_StylusUp);
+				RTSDI_StylusInRange | RTSDI_StylusDown | RTSDI_StylusUp);
 
-			DebugEvent("Requested data for Error, RealTimeStylusEnabled, RealTimeStylusDisabled, InAirPackets, Packets, StylusDown, StylusUp "
-				"from stylus input (DataInterest event).");
+			DebugEvent("Requested data for Error, RealTimeStylusEnabled, RealTimeStylusDisabled, InAirPackets, Packets, StylusInRange, StylusDown, "
+				"StylusUp from stylus input (DataInterest event).");
 		}
 
 		return S_OK;
@@ -684,6 +686,12 @@ namespace UE::StylusInput::Private::RealTimeStylus
 		return UpdateTabletContexts(RealTimeStylus, TabletContextIDsCount, TabletContextIDs);
 	}
 
+	HRESULT FRealTimeStylusPluginBase::ProcessStylusInRange(IRealTimeStylus* RealTimeStylus, TABLET_CONTEXT_ID TabletContextID, STYLUS_ID StylusID)
+	{
+		UpdateStylusInfoCallback.Execute(StylusID);
+		return S_OK;
+	}
+
 	HRESULT FRealTimeStylusPluginBase::ProcessTabletAdded(IInkTablet* Tablet)
 	{
 		// After a TabletAdded event, a RealTimeStylusDisabled event will fire directly followed by a RealTimeStylusEnabled event.
diff --git a/Engine/Plugins/Editor/StylusInput/Source/StylusInput/Private/Windows/RealTimeStylus/RealTimeStylusPluginBase.h b/Engine/Plugins/Editor/StylusInput/Source/StylusInput/Private/Windows/RealTimeStylus/RealTimeStylusPluginBase.h
--- a/Engine/Plugins/Editor/StylusInput/Source/StylusInput/Private/Windows/RealTimeStylus/RealTimeStylusPluginBase.h
+++ b/Engine/Plugins/Editor/StylusInput/Source/StylusInput/Private/Windows/RealTimeStylus/RealTimeStylusPluginBase.h
@@ -29,6 +29,7 @@ namespace UE::StylusInput::Private::RealTimeStylus
 
 	DECLARE_DELEGATE_RetVal(const FWindowContext&, FGetWindowContextCallback);
 	DECLARE_DELEGATE_OneParam(FUpdateTabletContextsCallback, const FTabletContextContainer&);
+	DECLARE_DELEGATE_OneParam(FUpdateStylusInfoCallback, STYLUS_ID);
 
 	class FRealTimeStylusPluginBase
 	{
@@ -41,7 +42,7 @@ namespace UE::StylusInput::Private::RealTimeStylus
 
 	protected:
 		FRealTimeStylusPluginBase(IStylusInputInstance* Instance, FGetWindowContextCallback&& GetWindowContextCallback,
-		                          FUpdateTabletContextsCallback&& UpdateTabletContextsCallback);
+		                          FUpdateTabletContextsCallback&& UpdateTabletContextsCallback, FUpdateStylusInfoCallback&& UpdateStylusInfoCallback);
 		~FRealTimeStylusPluginBase() = default;
 
 		virtual FString GetName() const = 0;
@@ -53,6 +54,7 @@ namespace UE::StylusInput::Private::RealTimeStylus
 		HRESULT ProcessPackets(const StylusInfo* StylusInfo, uint32 PacketCount, uint32 PacketBufferLength, EPacketType Type, const int32* PacketBuffer);
 		HRESULT ProcessRealTimeStylusEnabled(IRealTimeStylus* RealTimeStylus, uint32 TabletContextIDsCount, const TABLET_CONTEXT_ID* TabletContextIDs);
 		HRESULT ProcessRealTimeStylusDisabled(IRealTimeStylus* RealTimeStylus, uint32 TabletContextIDsCount, const TABLET_CONTEXT_ID* TabletContextIDs);
+		HRESULT ProcessStylusInRange(IRealTimeStylus* RealTimeStylus, TABLET_CONTEXT_ID TabletContextID, STYLUS_ID StylusID);
 		HRESULT ProcessTabletAdded(IInkTablet* Tablet);
 		HRESULT ProcessTabletRemoved(LONG TabletIndex);
 
@@ -64,6 +66,7 @@ namespace UE::StylusInput::Private::RealTimeStylus
 		FPacketStats PacketStats;
 		FGetWindowContextCallback GetWindowContextCallback;
 		FUpdateTabletContextsCallback UpdateTabletContextsCallback;
+		FUpdateStylusInfoCallback UpdateStylusInfoCallback;
 		FTabletContextContainer TabletContexts;
 		TArray<IStylusInputEventHandler*> EventHandlers;
 	};
diff --git a/Engine/Plugins/Editor/StylusInput/Source/StylusInput/Private/Windows/RealTimeStylus/RealTimeStylusPluginSync.cpp b/Engine/Plugins/Editor/StylusInput/Source/StylusInput/Private/Windows/RealTimeStylus/RealTimeStylusPluginSync.cpp
--- a/Engine/Plugins/Editor/StylusInput/Source/StylusInput/Private/Windows/RealTimeStylus/RealTimeStylusPluginSync.cpp
+++ b/Engine/Plugins/Editor/StylusInput/Source/StylusInput/Private/Windows/RealTimeStylus/RealTimeStylusPluginSync.cpp
@@ -11,8 +11,9 @@ namespace UE::StylusInput::Private::RealTimeStylus
 	FRealTimeStylusPluginSync::FRealTimeStylusPluginSync(IStylusInputInstance* Instance,
 	                                                     FGetWindowContextCallback&& GetWindowContext,
 	                                                     FUpdateTabletContextsCallback&& UpdateTabletContextsCallback,
+	                                                     FUpdateStylusInfoCallback UpdateStylusInfoCallback,
 	                                                     IStylusInputEventHandler* EventHandler)
-		: FRealTimeStylusPluginBase(Instance, MoveTemp(GetWindowContext), MoveTemp(UpdateTabletContextsCallback))
+		: FRealTimeStylusPluginBase(Instance, MoveTemp(GetWindowContext), MoveTemp(UpdateTabletContextsCallback), MoveTemp(UpdateStylusInfoCallback))
 	{
 		if (EventHandler)
 		{
@@ -85,7 +86,7 @@ namespace UE::StylusInput::Private::RealTimeStylus
 
 	HRESULT FRealTimeStylusPluginSync::StylusInRange(IRealTimeStylus* RealTimeStylus, TABLET_CONTEXT_ID TabletContextID, STYLUS_ID StylusID)
 	{
-		return E_NOTIMPL;
+		return ProcessStylusInRange(RealTimeStylus, TabletContextID, StylusID);
 	}
 
 	HRESULT FRealTimeStylusPluginSync::StylusOutOfRange(IRealTimeStylus* RealTimeStylus, TABLET_CONTEXT_ID TabletContextID, STYLUS_ID StylusID)
diff --git a/Engine/Plugins/Editor/StylusInput/Source/StylusInput/Private/Windows/RealTimeStylus/RealTimeStylusPluginSync.h b/Engine/Plugins/Editor/StylusInput/Source/StylusInput/Private/Windows/RealTimeStylus/RealTimeStylusPluginSync.h
--- a/Engine/Plugins/Editor/StylusInput/Source/StylusInput/Private/Windows/RealTimeStylus/RealTimeStylusPluginSync.h
+++ b/Engine/Plugins/Editor/StylusInput/Source/StylusInput/Private/Windows/RealTimeStylus/RealTimeStylusPluginSync.h
@@ -14,7 +14,8 @@ namespace UE::StylusInput::Private::RealTimeStylus
 	{
 	public:
 		FRealTimeStylusPluginSync(IStylusInputInstance* Instance, FGetWindowContextCallback&& GetWindowContext,
-		                              FUpdateTabletContextsCallback&& UpdateTabletContextsCallback, IStylusInputEventHandler* EventHandler);
+		                          FUpdateTabletContextsCallback&& UpdateTabletContextsCallback, FUpdateStylusInfoCallback UpdateStylusInfoCallback,
+		                          IStylusInputEventHandler* EventHandler);
 		~FRealTimeStylusPluginSync();
 
 		HRESULT CreateFreeThreadMarshaler();
