StylusInput: Fix crash when RealTimeStylus interface is not able to return a valid COM object

diff --git a/Engine/Plugins/Editor/StylusInput/Source/StylusInput/Private/Windows/RealTimeStylus/RealTimeStylusInstance.cpp b/Engine/Plugins/Editor/StylusInput/Source/StylusInput/Private/Windows/RealTimeStylus/RealTimeStylusInstance.cpp
--- a/Engine/Plugins/Editor/StylusInput/Source/StylusInput/Private/Windows/RealTimeStylus/RealTimeStylusInstance.cpp
+++ b/Engine/Plugins/Editor/StylusInput/Source/StylusInput/Private/Windows/RealTimeStylus/RealTimeStylusInstance.cpp
@@ -26,18 +26,58 @@ namespace UE::StylusInput::Private::RealTimeStylus
 		: ID(ID)
 		, RealTimeStylusAPI(FRealTimeStylusAPI::GetInstance())
 	{
-		Init(OSWindowHandle);
+		SetupWindowContext(OSWindowHandle);
+
+		void* OutInstance = nullptr;
+		if (Failed(RealTimeStylusAPI.CoCreateInstance(CLSID_RealTimeStylus, nullptr, CLSCTX_ALL, IID_IRealTimeStylus, &OutInstance), LOG_PREAMBLE))
+		{
+			LogError(LOG_PREAMBLE, "Could not create RealTimeStylus COM object instance!");
+		}
+		RealTimeStylus = static_cast<IRealTimeStylus*>(OutInstance);
+
+		if (RealTimeStylus.IsValid())
+		{
+			TArray<GUID> DesiredPackets;
+			DesiredPackets.Reserve(std::size(PacketPropertyConstants));
+			Algo::Transform(PacketPropertyConstants, DesiredPackets, [](const FPacketPropertyConstant& Property) { return Property.Guid; });
+
+			if (Failed(RealTimeStylus->SetDesiredPacketDescription(DesiredPackets.Num(), DesiredPackets.GetData()), LOG_PREAMBLE))
+			{
+				LogError(LOG_PREAMBLE, "Could not set desired packet description!");
+			}
+
+			if (Failed(RealTimeStylus->put_HWND(reinterpret_cast<HANDLE_PTR>(OSWindowHandle)), LOG_PREAMBLE))
+			{
+				LogError(LOG_PREAMBLE, "Could not set window handle!");
+			}
+
+			if (Failed(RealTimeStylus->put_Enabled(::Windows::TRUE), LOG_PREAMBLE))
+			{
+				LogError(LOG_PREAMBLE, "Could not enable real time stylus input!");
+			}
+
+			LogVerbose(LOG_PREAMBLE, "Successfully initialized RealTimeStylus COM object.");
+		}
 	}
 
 	FRealTimeStylusInstance::~FRealTimeStylusInstance()
 	{
-		RealTimeStylus->RemoveAllStylusAsyncPlugins();
-		RealTimeStylus->RemoveAllStylusSyncPlugins();
-		RealTimeStylus->Release();
+		if (RealTimeStylus.IsValid())
+		{
+			RealTimeStylus->RemoveAllStylusAsyncPlugins();
+			RealTimeStylus->RemoveAllStylusSyncPlugins();
+			RealTimeStylus->Release();
+		}
 	}
 
 	bool FRealTimeStylusInstance::AddEventHandler(IStylusInputEventHandler* EventHandler, const EEventHandlerThread Thread)
 	{
+		if (!RealTimeStylus.IsValid())
+		{
+			LogWarning(LOG_PREAMBLE, "Stylus input via the RealTimeStylus interface is not available; event handler was not installed.");
+			return false;
+		}
+
 		if (!EventHandler)
 		{
 			LogWarning(LOG_PREAMBLE, "Tried to add nullptr as event handler.");
@@ -89,6 +129,11 @@ namespace UE::StylusInput::Private::RealTimeStylus
 
 	bool FRealTimeStylusInstance::RemoveEventHandler(IStylusInputEventHandler* EventHandler)
 	{
+		if (!RealTimeStylus.IsValid())
+		{
+			return false;
+		}
+
 		if (!EventHandler)
 		{
 			LogWarning(LOG_PREAMBLE, "Tried to remove nullptr event handler.");
@@ -236,7 +281,10 @@ namespace UE::StylusInput::Private::RealTimeStylus
 
 	void FRealTimeStylusInstance::UpdateStylusInfo(STYLUS_ID StylusID)
 	{
-		LogWarning(LOG_PREAMBLE, "FRealTimeStylusInstance::UpdateStylusInfo()");
+		if (!ensure(RealTimeStylus.IsValid()))
+		{
+			return;
+		}
 
 		TSharedPtr<FStylusInfo> StylusInfo = MakeShared<FStylusInfo>(StylusID);
 
@@ -334,47 +382,12 @@ namespace UE::StylusInput::Private::RealTimeStylus
 		TabletContexts.Update(InTabletContexts);
 	}
 
-	void FRealTimeStylusInstance::Init(const HWND OSWindowHandle)
+	void FRealTimeStylusInstance::EnablePlugin(const EEventHandlerThread EventHandlerThread, IStylusInputEventHandler* EventHandler)
 	{
-		check(!RealTimeStylus.IsValid());
-
-		SetupWindowContext(OSWindowHandle);
-
-		void* OutInstance = nullptr;
-		if (Failed(RealTimeStylusAPI.CoCreateInstance(CLSID_RealTimeStylus, nullptr, CLSCTX_ALL, IID_IRealTimeStylus, &OutInstance), LOG_PREAMBLE))
-		{
-			LogError(LOG_PREAMBLE, "Could not create RealTimeStylus COM object instance!");
-		}
-		RealTimeStylus = static_cast<IRealTimeStylus*>(OutInstance);
-
-		if (RealTimeStylus.IsValid())
+		if (!ensure(RealTimeStylus.IsValid()))
 		{
-			TArray<GUID> DesiredPackets;
-			DesiredPackets.Reserve(std::size(PacketPropertyConstants));
-			Algo::Transform(PacketPropertyConstants, DesiredPackets, [](const FPacketPropertyConstant& Property) { return Property.Guid; });
-
-			if (Failed(RealTimeStylus->SetDesiredPacketDescription(DesiredPackets.Num(), DesiredPackets.GetData()), LOG_PREAMBLE))
-			{
-				LogError(LOG_PREAMBLE, "Could not set desired packet description!");
-			}
-
-			if (Failed(RealTimeStylus->put_HWND(reinterpret_cast<HANDLE_PTR>(OSWindowHandle)), LOG_PREAMBLE))
-			{
-				LogError(LOG_PREAMBLE, "Could not set window handle!");
-			}
-
-			if (Failed(RealTimeStylus->put_Enabled(::Windows::TRUE), LOG_PREAMBLE))
-			{
-				LogError(LOG_PREAMBLE, "Could not enable real time stylus input!");
-			}
-
-			LogVerbose(LOG_PREAMBLE, "Successfully initialized RealTimeStylus COM object.");
+			return;
 		}
-	}
-
-	void FRealTimeStylusInstance::EnablePlugin(const EEventHandlerThread EventHandlerThread, IStylusInputEventHandler* EventHandler)
-	{
-		check(RealTimeStylus.IsValid());
 
 		if (EventHandlerThread == EEventHandlerThread::OnGameThread)
 		{
@@ -429,7 +442,10 @@ namespace UE::StylusInput::Private::RealTimeStylus
 
 	void FRealTimeStylusInstance::DisablePlugin(const EEventHandlerThread EventHandlerThread)
 	{
-		check(RealTimeStylus.IsValid());
+		if (!ensure(RealTimeStylus.IsValid()))
+		{
+			return;
+		}
 
 #pragma warning(push)
 #pragma warning(disable : 6387)
diff --git a/Engine/Plugins/Editor/StylusInput/Source/StylusInput/Private/Windows/RealTimeStylus/RealTimeStylusInstance.h b/Engine/Plugins/Editor/StylusInput/Source/StylusInput/Private/Windows/RealTimeStylus/RealTimeStylusInstance.h
--- a/Engine/Plugins/Editor/StylusInput/Source/StylusInput/Private/Windows/RealTimeStylus/RealTimeStylusInstance.h
+++ b/Engine/Plugins/Editor/StylusInput/Source/StylusInput/Private/Windows/RealTimeStylus/RealTimeStylusInstance.h
@@ -45,8 +45,6 @@ namespace UE::StylusInput::Private::RealTimeStylus
 		virtual FText GetName() override;
 
 	private:
-		void Init(HWND OSWindowHandle);
-
 		void EnablePlugin(EEventHandlerThread EventHandlerThread, IStylusInputEventHandler* EventHandler);
 		void DisablePlugin(EEventHandlerThread EventHandlerThread);
 
