StylusInput: Add check for successful initialization of an interface

diff --git a/Engine/Plugins/Editor/StylusInput/Source/StylusInput/Private/StylusInput.cpp b/Engine/Plugins/Editor/StylusInput/Source/StylusInput/Private/StylusInput.cpp
--- a/Engine/Plugins/Editor/StylusInput/Source/StylusInput/Private/StylusInput.cpp
+++ b/Engine/Plugins/Editor/StylusInput/Source/StylusInput/Private/StylusInput.cpp
@@ -52,7 +52,16 @@ namespace UE::StylusInput
 					return nullptr;
 				}
 
-				return Interface ? Interface->CreateInstance(Window) : nullptr;
+				IStylusInputInstance* Instance = Interface ? Interface->CreateInstance(Window) : nullptr;
+
+				if (Instance && !Instance->WasInitializedSuccessfully())
+				{
+					UE_LOG(LogStylusInput, Error, TEXT("Platform interface instance was not initialized successfully."));
+					Interface->ReleaseInstance(Instance);
+					Instance = nullptr;
+				}
+
+				return Instance;
 			}
 
 			bool ReleaseInstance(IStylusInputInstance* Instance) const
diff --git a/Engine/Plugins/Editor/StylusInput/Source/StylusInput/Private/Windows/RealTimeStylus/RealTimeStylusInstance.cpp b/Engine/Plugins/Editor/StylusInput/Source/StylusInput/Private/Windows/RealTimeStylus/RealTimeStylusInstance.cpp
--- a/Engine/Plugins/Editor/StylusInput/Source/StylusInput/Private/Windows/RealTimeStylus/RealTimeStylusInstance.cpp
+++ b/Engine/Plugins/Editor/StylusInput/Source/StylusInput/Private/Windows/RealTimeStylus/RealTimeStylusInstance.cpp
@@ -43,19 +43,30 @@ namespace UE::StylusInput::Private::RealTimeStylus
 
 			if (Failed(RealTimeStylus->SetDesiredPacketDescription(DesiredPackets.Num(), DesiredPackets.GetData()), LOG_PREAMBLE))
 			{
-				LogError(LOG_PREAMBLE, "Could not set desired packet description!");
+				LogError(LOG_PREAMBLE, "Could not set desired RealTimeStylus packet description!");
+				bWasInitializedSuccessfully = false;
 			}
 
 			if (Failed(RealTimeStylus->put_HWND(reinterpret_cast<HANDLE_PTR>(OSWindowHandle)), LOG_PREAMBLE))
 			{
-				LogError(LOG_PREAMBLE, "Could not set window handle!");
+				LogError(LOG_PREAMBLE, "Could not set RealTimeStylus window handle!");
+				bWasInitializedSuccessfully = false;
 			}
 
 			if (Failed(RealTimeStylus->put_Enabled(::Windows::TRUE), LOG_PREAMBLE))
 			{
-				LogError(LOG_PREAMBLE, "Could not enable real time stylus input!");
+				LogError(LOG_PREAMBLE, "Could not enable RealTimeStylus interface!");
+				bWasInitializedSuccessfully = false;
 			}
+		}
+		else
+		{
+			LogError(LOG_PREAMBLE, "Stylus input via the RealTimeStylus interface is not available!");
+			bWasInitializedSuccessfully = false;
+		}
 
+		if (bWasInitializedSuccessfully)
+		{
 			LogVerbose(LOG_PREAMBLE, "Successfully initialized RealTimeStylus COM object.");
 		}
 	}
@@ -200,6 +211,11 @@ namespace UE::StylusInput::Private::RealTimeStylus
 		return FText::Format(LOCTEXT("RealTimeStylus", "RealTimeStylus #{0}"), ID);
 	}
 
+	bool FRealTimeStylusInstance::WasInitializedSuccessfully()
+	{
+		return bWasInitializedSuccessfully;
+	}
+
 	void FRealTimeStylusInstance::SetupWindowContext(const HWND HWindow)
 	{
 		RECT WindowRectangle;
diff --git a/Engine/Plugins/Editor/StylusInput/Source/StylusInput/Private/Windows/RealTimeStylus/RealTimeStylusInstance.h b/Engine/Plugins/Editor/StylusInput/Source/StylusInput/Private/Windows/RealTimeStylus/RealTimeStylusInstance.h
--- a/Engine/Plugins/Editor/StylusInput/Source/StylusInput/Private/Windows/RealTimeStylus/RealTimeStylusInstance.h
+++ b/Engine/Plugins/Editor/StylusInput/Source/StylusInput/Private/Windows/RealTimeStylus/RealTimeStylusInstance.h
@@ -44,6 +44,8 @@ namespace UE::StylusInput::Private::RealTimeStylus
 
 		virtual FText GetName() override;
 
+		virtual bool WasInitializedSuccessfully() override;
+
 	private:
 		void EnablePlugin(EEventHandlerThread EventHandlerThread, IStylusInputEventHandler* EventHandler);
 		void DisablePlugin(EEventHandlerThread EventHandlerThread);
@@ -64,6 +66,8 @@ namespace UE::StylusInput::Private::RealTimeStylus
 
 		FTabletContextThreadSafeContainer TabletContexts;
 		FStylusInfoThreadSafeContainer StylusInfos;
+
+		bool bWasInitializedSuccessfully = true;
 	};
 }
 
diff --git a/Engine/Plugins/Editor/StylusInput/Source/StylusInput/Private/Windows/Wintab/WintabInstance.cpp b/Engine/Plugins/Editor/StylusInput/Source/StylusInput/Private/Windows/Wintab/WintabInstance.cpp
--- a/Engine/Plugins/Editor/StylusInput/Source/StylusInput/Private/Windows/Wintab/WintabInstance.cpp
+++ b/Engine/Plugins/Editor/StylusInput/Source/StylusInput/Private/Windows/Wintab/WintabInstance.cpp
@@ -647,6 +647,11 @@ namespace UE::StylusInput::Private::Wintab
 		return FText::Format(LOCTEXT("Wintab", "Wintab #{0}"), ID);
 	}
 
+	bool FWintabInstance::WasInitializedSuccessfully()
+	{
+		return true;
+	}
+
 	const FTabletContext* FWintabInstance::GetTabletContextInternal(const uint32 TabletContextID) const
 	{
 		return TabletContexts.Get(TabletContextID).Get();
diff --git a/Engine/Plugins/Editor/StylusInput/Source/StylusInput/Private/Windows/Wintab/WintabInstance.h b/Engine/Plugins/Editor/StylusInput/Source/StylusInput/Private/Windows/Wintab/WintabInstance.h
--- a/Engine/Plugins/Editor/StylusInput/Source/StylusInput/Private/Windows/Wintab/WintabInstance.h
+++ b/Engine/Plugins/Editor/StylusInput/Source/StylusInput/Private/Windows/Wintab/WintabInstance.h
@@ -40,6 +40,8 @@ namespace UE::StylusInput::Private::Wintab
 
 		virtual FText GetName() override;
 
+		virtual bool WasInitializedSuccessfully() override;
+
 	private:
 		const FTabletContext* GetTabletContextInternal(uint32 TabletContextID) const;
 		uint32 GetStylusID(uint32 TabletContextID, uint32 CursorIndex);
diff --git a/Engine/Plugins/Editor/StylusInput/Source/StylusInput/Public/StylusInput.h b/Engine/Plugins/Editor/StylusInput/Source/StylusInput/Public/StylusInput.h
--- a/Engine/Plugins/Editor/StylusInput/Source/StylusInput/Public/StylusInput.h
+++ b/Engine/Plugins/Editor/StylusInput/Source/StylusInput/Public/StylusInput.h
@@ -111,6 +111,12 @@ namespace UE::StylusInput
 		 * @returns Name of this instance.
 		 */
 		virtual FText GetName() = 0;
+
+		/**
+		 * 
+		 * @returns true if this instance was initialized successfully and able to provide stylus input events.
+		 */
+		virtual bool WasInitializedSuccessfully() = 0;
 	};
 
 	/**
