StylusInput: Enable Wintab API as default on Windows

diff --git a/Engine/Plugins/Editor/StylusInput/Source/StylusInput/Private/StylusInput.cpp b/Engine/Plugins/Editor/StylusInput/Source/StylusInput/Private/StylusInput.cpp
--- a/Engine/Plugins/Editor/StylusInput/Source/StylusInput/Private/StylusInput.cpp
+++ b/Engine/Plugins/Editor/StylusInput/Source/StylusInput/Private/StylusInput.cpp
@@ -4,6 +4,7 @@
 
 #include "StylusInputInterface.h"
 #include "StylusInputUtils.h"
+#include <Algo/Find.h>
 
 #if PLATFORM_WINDOWS
 	#if STYLUSINPUT_USE_WINTAB
@@ -16,86 +17,132 @@ namespace UE::StylusInput
 {
 	namespace Private
 	{
-		class FStylusInputImpl;
-
-		class FNotSupportedInterface final : public IStylusInputInterface
+		class FStylusInputImpl
 		{
 		public:
-			virtual ~FNotSupportedInterface() override = default;
+			FStylusInputImpl()
+			{
+#if PLATFORM_WINDOWS
 
-			virtual IStylusInputInstance* CreateInstance(SWindow& Window) override { return nullptr; }
+	#if STYLUSINPUT_USE_WINTAB && STYLUSINPUT_USE_WINTAB_AS_DEFAULT
+				AddInterface<Wintab::FWintabInterface>();
+	#endif
 
-			virtual bool ReleaseInstance(IStylusInputInstance* Instance) override { return false; }
+				AddInterface<RealTimeStylus::FRealTimeStylusInterface>();
 
-		private:
-			friend FStylusInputImpl;
-			static TUniquePtr<IStylusInputInterface> Create()
-			{
-				UE_LOG(LogStylusInput, Warning, TEXT("Stylus input is not supported for this platform."));
-				return nullptr;
+	#if STYLUSINPUT_USE_WINTAB && !STYLUSINPUT_USE_WINTAB_AS_DEFAULT
+				AddInterface<Wintab::FWintabInterface>();
+	#endif
+
+#elif PLATFORM_MAC
+
+				// TODO 
+
+#elif PLATFORM_LINUX
+
+				// TODO 
+
+#endif
 			}
-		};
 
-		class FStylusInputImpl
-		{
-		public:
-			FStylusInputImpl()
-				: Interface(SelectInterface())
+			TArray<FName> GetAvailableInterfaces()
 			{
+				TArray<FName> APINames;
+				APINames.Reserve(Interfaces.Num());
+
+				for (const TUniquePtr<IStylusInputInterface>& Interface : Interfaces)
+				{
+					APINames.Add(Interface->GetName());
+				}
+
+				return APINames;
 			}
 
-			IStylusInputInstance* CreateInstance(SWindow& Window) const
+			IStylusInputInstance* CreateInstance(SWindow& Window, const FName* InterfaceName, const bool bRequestedInterfaceOnly) const
 			{
-				if (!Interface)
+				// Helper for trying to create an instance for a given interface.
+				auto CreateInstanceForInterface = [&Window](IStylusInputInterface& Interface) -> IStylusInputInstance*
 				{
-					UE_LOG(LogStylusInput, Error, TEXT("Platform interface not available."));
+					if (IStylusInputInstance* Instance = Interface.CreateInstance(Window))
+					{
+						if (Instance->WasInitializedSuccessfully())
+						{
+							return Instance;
+						}
+
+						UE_LOG(LogStylusInput, Error, TEXT("Stylus input instance for interface '%s' was not initialized successfully."),
+						       *Interface.GetName().ToString());
+
+						Interface.ReleaseInstance(Instance);
+					}
+
 					return nullptr;
-				}
+				};
 
-				IStylusInputInstance* Instance = Interface ? Interface->CreateInstance(Window) : nullptr;
+				// Try to create an instance for the named interface.
+				if (InterfaceName)
+				{
+					const TUniquePtr<IStylusInputInterface>* RequestedInterface = Algo::FindByPredicate(
+						Interfaces, [InterfaceName](const TUniquePtr<IStylusInputInterface>& Interface) { return Interface->GetName() == *InterfaceName; });
+
+					if (RequestedInterface)
+					{
+						if (IStylusInputInstance* Instance = CreateInstanceForInterface(**RequestedInterface))
+						{
+							return Instance;
+						}
+					}
+					else
+					{
+						UE_LOG(LogStylusInput, Warning, TEXT("Requested stylus input interface '%s' is not available."), *InterfaceName->ToString());
+					}
+				}
 
-				if (Instance && !Instance->WasInitializedSuccessfully())
+				// Otherwise, try to create an instance for any available interface, unless only the requested interface is valid.
+				if (!InterfaceName || !bRequestedInterfaceOnly)
 				{
-					UE_LOG(LogStylusInput, Error, TEXT("Platform interface instance was not initialized successfully."));
-					Interface->ReleaseInstance(Instance);
-					Instance = nullptr;
+					for (const TUniquePtr<IStylusInputInterface>& Interface : Interfaces)
+					{
+						if (Interface)
+						{
+							if (IStylusInputInstance* Instance = CreateInstanceForInterface(*Interface))
+							{
+								return Instance;
+							}
+						}
+					}
 				}
 
-				return Instance;
+				UE_LOG(LogStylusInput, Warning, TEXT("No valid stylus input interface and/or instance available."));
+				return nullptr;
 			}
 
 			bool ReleaseInstance(IStylusInputInstance* Instance) const
 			{
-				return Interface ? Interface->ReleaseInstance(Instance) : false;
-			}
-
-		private:
-			static TUniquePtr<IStylusInputInterface> SelectInterface()
-			{
-#if PLATFORM_WINDOWS
-	#if STYLUSINPUT_USE_WINTAB
+				const FName InterfaceName = Instance->GetInterfaceName();
 
-				// Prefer using Wintab, and fall back to RealTimeStylus if not available.
-				TUniquePtr<IStylusInputInterface> Interface = Wintab::FWintabInterface::Create();
-				if (!Interface)
+				for (const TUniquePtr<IStylusInputInterface>& Interface : Interfaces)
 				{
-					Interface = RealTimeStylus::FRealTimeStylusInterface::Create();
+					if (Interface->GetName() == InterfaceName)
+					{
+						return Interface->ReleaseInstance(Instance);
+					}
 				}
-				return Interface;
 
-	#else
-
-				return RealTimeStylus::FRealTimeStylusInterface::Create();
-
-	#endif
-#else
+				return false;
+			}
 
-				return FNotSupportedInterface::Create();
+		private:
+			TArray<TUniquePtr<IStylusInputInterface>> Interfaces;
 
-#endif
+			template <typename InterfaceType>
+			void AddInterface()
+			{
+				if (TUniquePtr<IStylusInputInterface> Interface = InterfaceType::Create())
+				{
+					Interfaces.Add(MoveTemp(Interface));
+				}
 			}
-
-			TUniquePtr<IStylusInputInterface> Interface;
 		};
 	}
 
@@ -105,16 +152,26 @@ namespace UE::StylusInput
 		return ImplSingleton;
 	}
 
+	TArray<FName> GetAvailableInterfaces()
+	{
+		return GetImplSingleton().GetAvailableInterfaces();
+	}
+
 	IStylusInputInstance* CreateInstance(SWindow& Window)
 	{
-		return GetImplSingleton().CreateInstance(Window);
+		return GetImplSingleton().CreateInstance(Window, nullptr, false);
+	}
+
+	IStylusInputInstance* CreateInstance(SWindow& Window, const FName Interface, const bool bRequestedInterfaceOnly)
+	{
+		return GetImplSingleton().CreateInstance(Window, &Interface, bRequestedInterfaceOnly);
 	}
 
 	bool ReleaseInstance(IStylusInputInstance* Instance)
 	{
 		if (!Instance)
 		{
-			UE_LOG(Private::LogStylusInput, Warning, TEXT("Nullptr passed into FStylusInput::ReleaseInstance()."));
+			UE_LOG(Private::LogStylusInput, Warning, TEXT("Nullptr passed into StylusInput::ReleaseInstance()."));
 			return false;
 		}
 
diff --git a/Engine/Plugins/Editor/StylusInput/Source/StylusInput/Private/StylusInputInterface.h b/Engine/Plugins/Editor/StylusInput/Source/StylusInput/Private/StylusInputInterface.h
--- a/Engine/Plugins/Editor/StylusInput/Source/StylusInput/Private/StylusInputInterface.h
+++ b/Engine/Plugins/Editor/StylusInput/Source/StylusInput/Private/StylusInputInterface.h
@@ -4,6 +4,7 @@
 
 #include <HAL/Platform.h>
 
+class FName;
 class SWindow;
 
 namespace UE::StylusInput
@@ -19,6 +20,8 @@ namespace UE::StylusInput::Private
 	public:
 		virtual ~IStylusInputInterface() = default;
 
+		virtual FName GetName() = 0;
+
 		virtual IStylusInputInstance* CreateInstance(SWindow& Window) = 0;
 		virtual bool ReleaseInstance(IStylusInputInstance* Instance) = 0;
 	};
diff --git a/Engine/Plugins/Editor/StylusInput/Source/StylusInput/Private/Windows/RealTimeStylus/RealTimeStylusAPI.cpp b/Engine/Plugins/Editor/StylusInput/Source/StylusInput/Private/Windows/RealTimeStylus/RealTimeStylusAPI.cpp
--- a/Engine/Plugins/Editor/StylusInput/Source/StylusInput/Private/Windows/RealTimeStylus/RealTimeStylusAPI.cpp
+++ b/Engine/Plugins/Editor/StylusInput/Source/StylusInput/Private/Windows/RealTimeStylus/RealTimeStylusAPI.cpp
@@ -125,6 +125,12 @@ namespace UE::StylusInput::Private::RealTimeStylus
 		}
 	}
 
+	FName FRealTimeStylusAPI::GetName()
+	{
+		static FName Name("RealTimeStylus");
+		return Name;
+	}
+
 	const FRealTimeStylusAPI& FRealTimeStylusAPI::GetInstance()
 	{
 		static FRealTimeStylusAPI API;
diff --git a/Engine/Plugins/Editor/StylusInput/Source/StylusInput/Private/Windows/RealTimeStylus/RealTimeStylusAPI.h b/Engine/Plugins/Editor/StylusInput/Source/StylusInput/Private/Windows/RealTimeStylus/RealTimeStylusAPI.h
--- a/Engine/Plugins/Editor/StylusInput/Source/StylusInput/Private/Windows/RealTimeStylus/RealTimeStylusAPI.h
+++ b/Engine/Plugins/Editor/StylusInput/Source/StylusInput/Private/Windows/RealTimeStylus/RealTimeStylusAPI.h
@@ -19,6 +19,7 @@ namespace UE::StylusInput::Private::RealTimeStylus
 	class FRealTimeStylusAPI
 	{
 	public:
+		static FName GetName();
 		static const FRealTimeStylusAPI& GetInstance();
 
 		bool SatisfiesRequirements() const;
diff --git a/Engine/Plugins/Editor/StylusInput/Source/StylusInput/Private/Windows/RealTimeStylus/RealTimeStylusInstance.cpp b/Engine/Plugins/Editor/StylusInput/Source/StylusInput/Private/Windows/RealTimeStylus/RealTimeStylusInstance.cpp
--- a/Engine/Plugins/Editor/StylusInput/Source/StylusInput/Private/Windows/RealTimeStylus/RealTimeStylusInstance.cpp
+++ b/Engine/Plugins/Editor/StylusInput/Source/StylusInput/Private/Windows/RealTimeStylus/RealTimeStylusInstance.cpp
@@ -206,6 +206,11 @@ namespace UE::StylusInput::Private::RealTimeStylus
 		return -1.0f;
 	}
 
+	FName FRealTimeStylusInstance::GetInterfaceName()
+	{
+		return FRealTimeStylusAPI::GetName();
+	}
+
 	FText FRealTimeStylusInstance::GetName()
 	{
 		return FText::Format(LOCTEXT("RealTimeStylus", "RealTimeStylus #{0}"), ID);
diff --git a/Engine/Plugins/Editor/StylusInput/Source/StylusInput/Private/Windows/RealTimeStylus/RealTimeStylusInstance.h b/Engine/Plugins/Editor/StylusInput/Source/StylusInput/Private/Windows/RealTimeStylus/RealTimeStylusInstance.h
--- a/Engine/Plugins/Editor/StylusInput/Source/StylusInput/Private/Windows/RealTimeStylus/RealTimeStylusInstance.h
+++ b/Engine/Plugins/Editor/StylusInput/Source/StylusInput/Private/Windows/RealTimeStylus/RealTimeStylusInstance.h
@@ -42,6 +42,7 @@ namespace UE::StylusInput::Private::RealTimeStylus
 
 		virtual float GetPacketsPerSecond(EEventHandlerThread EventHandlerThread) const override;
 
+		virtual FName GetInterfaceName() override;
 		virtual FText GetName() override;
 
 		virtual bool WasInitializedSuccessfully() override;
diff --git a/Engine/Plugins/Editor/StylusInput/Source/StylusInput/Private/Windows/RealTimeStylus/RealTimeStylusInterface.cpp b/Engine/Plugins/Editor/StylusInput/Source/StylusInput/Private/Windows/RealTimeStylus/RealTimeStylusInterface.cpp
--- a/Engine/Plugins/Editor/StylusInput/Source/StylusInput/Private/Windows/RealTimeStylus/RealTimeStylusInterface.cpp
+++ b/Engine/Plugins/Editor/StylusInput/Source/StylusInput/Private/Windows/RealTimeStylus/RealTimeStylusInterface.cpp
@@ -20,6 +20,11 @@
 
 namespace UE::StylusInput::Private::RealTimeStylus
 {
+	FName FRealTimeStylusInterface::GetName()
+	{
+		return FRealTimeStylusAPI::GetName();
+	}
+
 	IStylusInputInstance* FRealTimeStylusInterface::CreateInstance(SWindow& Window)
 	{
 		if (FRefCountedInstance* ExistingRefCountedInstance = Instances.Find(&Window))
diff --git a/Engine/Plugins/Editor/StylusInput/Source/StylusInput/Private/Windows/RealTimeStylus/RealTimeStylusInterface.h b/Engine/Plugins/Editor/StylusInput/Source/StylusInput/Private/Windows/RealTimeStylus/RealTimeStylusInterface.h
--- a/Engine/Plugins/Editor/StylusInput/Source/StylusInput/Private/Windows/RealTimeStylus/RealTimeStylusInterface.h
+++ b/Engine/Plugins/Editor/StylusInput/Source/StylusInput/Private/Windows/RealTimeStylus/RealTimeStylusInterface.h
@@ -23,6 +23,8 @@ namespace UE::StylusInput::Private::RealTimeStylus
 	class FRealTimeStylusInterface : public IStylusInputInterface
 	{
 	public:
+		virtual FName GetName() override;
+
 		virtual IStylusInputInstance* CreateInstance(SWindow& Window) override;
 		virtual bool ReleaseInstance(IStylusInputInstance* Instance) override;
 
diff --git a/Engine/Plugins/Editor/StylusInput/Source/StylusInput/Private/Windows/Wintab/WintabAPI.cpp b/Engine/Plugins/Editor/StylusInput/Source/StylusInput/Private/Windows/Wintab/WintabAPI.cpp
--- a/Engine/Plugins/Editor/StylusInput/Source/StylusInput/Private/Windows/Wintab/WintabAPI.cpp
+++ b/Engine/Plugins/Editor/StylusInput/Source/StylusInput/Private/Windows/Wintab/WintabAPI.cpp
@@ -5,6 +5,7 @@
 #if PLATFORM_WINDOWS && STYLUSINPUT_USE_WINTAB
 
 #include <StylusInputUtils.h>
+#include <Windows/WindowsPlatformProcess.h>
 
 #define LOG_PREAMBLE "WintabAPI"
 
@@ -135,6 +136,12 @@ namespace UE::StylusInput::Private::Wintab
 		}
 	}
 
+	FName FWintabAPI::GetName()
+	{
+		static FName Name("Wintab");
+		return Name;
+	}
+
 	const FWintabAPI& FWintabAPI::GetInstance()
 	{
 		static FWintabAPI WintabAPI;
diff --git a/Engine/Plugins/Editor/StylusInput/Source/StylusInput/Private/Windows/Wintab/WintabAPI.h b/Engine/Plugins/Editor/StylusInput/Source/StylusInput/Private/Windows/Wintab/WintabAPI.h
--- a/Engine/Plugins/Editor/StylusInput/Source/StylusInput/Private/Windows/Wintab/WintabAPI.h
+++ b/Engine/Plugins/Editor/StylusInput/Source/StylusInput/Private/Windows/Wintab/WintabAPI.h
@@ -10,13 +10,19 @@
 #include <Containers/UnrealString.h>
 #include <Templates/UniquePtr.h>
 
-#define WIN32_LEAN_AND_MEAN
-#include <windows.h>
-#include <WINTAB.H>
-#define PACKETDATA (/*PK_CONTEXT |*/ PK_STATUS | PK_TIME | /*PK_CHANGED |*/ PK_SERIAL_NUMBER | PK_CURSOR | PK_BUTTONS | PK_X | PK_Y | PK_Z | PK_NORMAL_PRESSURE | PK_TANGENT_PRESSURE | PK_ORIENTATION | PK_ROTATION)
-#define PACKETMODE PK_BUTTONS
-#include <PKTDEF.H>
-#undef INT // WINTAB.H defines an INT macro that collides with a macro definition in AllowWindowsPlatformTypes.h when using Unity builds.
+#include <Windows/AllowWindowsPlatformTypes.h>
+
+	#define WIN32_LEAN_AND_MEAN
+	#include <windows.h>
+
+	#undef INT // WINTAB.H defines an INT macro that collides with a macro definition in AllowWindowsPlatformTypes.h when using Unity builds.
+		#include <WINTAB.H>
+		#define PACKETDATA (/*PK_CONTEXT |*/ PK_STATUS | PK_TIME | /*PK_CHANGED |*/ PK_SERIAL_NUMBER | PK_CURSOR | PK_BUTTONS | PK_X | PK_Y | PK_Z | PK_NORMAL_PRESSURE | PK_TANGENT_PRESSURE | PK_ORIENTATION | PK_ROTATION)
+		#define PACKETMODE PK_BUTTONS
+		#include <PKTDEF.H>
+	//#undef INT // WINTAB.H defines an INT macro that collides with a macro definition in AllowWindowsPlatformTypes.h when using Unity builds.
+
+#include <Windows/HideWindowsPlatformTypes.h>
 
 #define WTINFO_WORKAROUND_FOR_OUTPUT_BUFFER_SIZE_BUG 0
 
@@ -42,6 +48,7 @@ namespace UE::StylusInput::Private::Wintab
 	class FWintabAPI
 	{
 	public:
+		static FName GetName();
 		static const FWintabAPI& GetInstance();
 
 		bool SatisfiesRequirements() const;
@@ -152,18 +159,15 @@ namespace UE::StylusInput::Private::Wintab
 }
 
 /*
- * TODO 
- * [X] Handle removing/adding tablet while running (evaluate WT_INFOCHANGE event)
- * [ ] PK_BUTTONS (plus RealtimeStylus buttons)
+ * TODO
+ * [X] Allow applications to request specific API when creating instance
+ * [ ] Make Wintab the default (with additional editor settings?)
+ * [ ] Propagate buttons data for both Wintab and Windows (PK_BUTTONS and/or StylusButtonDown/StylusButtonDown events)
  * [ ] Queue size
  * [ ] Query multiple packets at once
  * [ ] Async
  * [ ] Deal with cursors not having a physical ID
- * [ ] Make available by default (with additional editor settings)
- * [ ] Propagate buttons data for both Wintab and Windows (StylusButtonDown/StylusButtonDown events)
  * [ ] Allegedly, some Wintab libraries don't handle relative button presses, so maybe switch to absolute?
- * [X] Lost packet stats?
- * [ ] RealTimeStylus StylusInRange and StylusOutOfRange events (re-evaluate stylus info?)
  */
 
 #endif
diff --git a/Engine/Plugins/Editor/StylusInput/Source/StylusInput/Private/Windows/Wintab/WintabInstance.cpp b/Engine/Plugins/Editor/StylusInput/Source/StylusInput/Private/Windows/Wintab/WintabInstance.cpp
--- a/Engine/Plugins/Editor/StylusInput/Source/StylusInput/Private/Windows/Wintab/WintabInstance.cpp
+++ b/Engine/Plugins/Editor/StylusInput/Source/StylusInput/Private/Windows/Wintab/WintabInstance.cpp
@@ -603,6 +603,7 @@ namespace UE::StylusInput::Private::Wintab
 		if (Thread == EEventHandlerThread::Asynchronous)
 		{
 			// TODO Handle asynchronous thread
+			LogError(LOG_PREAMBLE, "Asynchronous event handler is not supported yet.");
 			return false;
 		}
 
@@ -642,6 +643,11 @@ namespace UE::StylusInput::Private::Wintab
         return MessageHandler.GetPacketStats().GetPacketsPerSecond();
     }
 
+	FName FWintabInstance::GetInterfaceName()
+	{
+		return FWintabAPI::GetName();
+	}
+
 	FText FWintabInstance::GetName()
 	{
 		return FText::Format(LOCTEXT("Wintab", "Wintab #{0}"), ID);
diff --git a/Engine/Plugins/Editor/StylusInput/Source/StylusInput/Private/Windows/Wintab/WintabInstance.h b/Engine/Plugins/Editor/StylusInput/Source/StylusInput/Private/Windows/Wintab/WintabInstance.h
--- a/Engine/Plugins/Editor/StylusInput/Source/StylusInput/Private/Windows/Wintab/WintabInstance.h
+++ b/Engine/Plugins/Editor/StylusInput/Source/StylusInput/Private/Windows/Wintab/WintabInstance.h
@@ -38,6 +38,7 @@ namespace UE::StylusInput::Private::Wintab
 
 		virtual float GetPacketsPerSecond(EEventHandlerThread EventHandlerThread) const override;
 
+		virtual FName GetInterfaceName() override;
 		virtual FText GetName() override;
 
 		virtual bool WasInitializedSuccessfully() override;
diff --git a/Engine/Plugins/Editor/StylusInput/Source/StylusInput/Private/Windows/Wintab/WintabInterface.cpp b/Engine/Plugins/Editor/StylusInput/Source/StylusInput/Private/Windows/Wintab/WintabInterface.cpp
--- a/Engine/Plugins/Editor/StylusInput/Source/StylusInput/Private/Windows/Wintab/WintabInterface.cpp
+++ b/Engine/Plugins/Editor/StylusInput/Source/StylusInput/Private/Windows/Wintab/WintabInterface.cpp
@@ -18,6 +18,11 @@
 
 namespace UE::StylusInput::Private::Wintab
 {
+	FName FWintabInterface::GetName()
+	{
+		return FWintabAPI::GetName();
+	}
+
 	IStylusInputInstance* FWintabInterface::CreateInstance(SWindow& Window)
 	{
 		if (FRefCountedInstance* ExistingRefCountedInstance = Instances.Find(&Window))
diff --git a/Engine/Plugins/Editor/StylusInput/Source/StylusInput/Private/Windows/Wintab/WintabInterface.h b/Engine/Plugins/Editor/StylusInput/Source/StylusInput/Private/Windows/Wintab/WintabInterface.h
--- a/Engine/Plugins/Editor/StylusInput/Source/StylusInput/Private/Windows/Wintab/WintabInterface.h
+++ b/Engine/Plugins/Editor/StylusInput/Source/StylusInput/Private/Windows/Wintab/WintabInterface.h
@@ -23,6 +23,8 @@ namespace UE::StylusInput::Private::Wintab
 	class FWintabInterface : public IStylusInputInterface
 	{
 	public:
+		virtual FName GetName() override;
+
 		virtual IStylusInputInstance* CreateInstance(SWindow& Window) override;
 		virtual bool ReleaseInstance(IStylusInputInstance* Instance) override;
 
diff --git a/Engine/Plugins/Editor/StylusInput/Source/StylusInput/Private/Windows/Wintab/WintabMessageHandler.cpp b/Engine/Plugins/Editor/StylusInput/Source/StylusInput/Private/Windows/Wintab/WintabMessageHandler.cpp
--- a/Engine/Plugins/Editor/StylusInput/Source/StylusInput/Private/Windows/Wintab/WintabMessageHandler.cpp
+++ b/Engine/Plugins/Editor/StylusInput/Source/StylusInput/Private/Windows/Wintab/WintabMessageHandler.cpp
@@ -141,7 +141,10 @@ bool FWintabMessageHandler::ProcessMessage(const HWND Hwnd, const uint32 Msg, co
 				// support. Just in case, we trigger a cursor change if either of the flags are set.
 
 				bCursorChange = true;
-				CursorChangeTableContextID = static_cast<uint32>(WParam);
+
+#if CHECK_MESSAGE_ORDER_ASSUMPTIONS
+				CursorChangeTabletContextID = static_cast<uint32>(WParam);
+#endif
 			}
 
 			bHandled = true;
@@ -165,8 +168,11 @@ bool FWintabMessageHandler::ProcessMessage(const HWND Hwnd, const uint32 Msg, co
 #endif
 
 			bCursorChange = true;
-			CursorChangeTableContextID = static_cast<uint32>(LParam);
+
+#if CHECK_MESSAGE_ORDER_ASSUMPTIONS
+			CursorChangeTabletContextID = static_cast<uint32>(LParam);
 			CursorChangeSerialNumber = static_cast<uint32>(WParam);
+#endif
 
 			bHandled = true;
 			break;
@@ -238,7 +244,7 @@ void FWintabMessageHandler::ProcessCursorChange(const uint32 TabletContextID, co
 	check(bCursorChange);
 
 #if CHECK_MESSAGE_ORDER_ASSUMPTIONS
-	if (CursorChangeTableContextID != TabletContextID || !(CursorChangeSerialNumber == 0 || CursorChangeSerialNumber == WintabPacket.pkSerialNumber))
+	if (CursorChangeTabletContextID != TabletContextID || !(CursorChangeSerialNumber == 0 || CursorChangeSerialNumber == WintabPacket.pkSerialNumber))
 	{
 		DebugEvent("Cursor change assumptions violated: The next packet is not matching the previous WT_PROXIMITY/WT_CSRCHANGE messages.");
 	}
@@ -249,7 +255,7 @@ void FWintabMessageHandler::ProcessCursorChange(const uint32 TabletContextID, co
 	bCursorChange = false;
 
 #if CHECK_MESSAGE_ORDER_ASSUMPTIONS
-	CursorChangeTableContextID = 0;
+	CursorChangeTabletContextID = 0;
 	CursorChangeSerialNumber = 0;
 #endif
 }
diff --git a/Engine/Plugins/Editor/StylusInput/Source/StylusInput/Private/Windows/Wintab/WintabMessageHandler.h b/Engine/Plugins/Editor/StylusInput/Source/StylusInput/Private/Windows/Wintab/WintabMessageHandler.h
--- a/Engine/Plugins/Editor/StylusInput/Source/StylusInput/Private/Windows/Wintab/WintabMessageHandler.h
+++ b/Engine/Plugins/Editor/StylusInput/Source/StylusInput/Private/Windows/Wintab/WintabMessageHandler.h
@@ -27,7 +27,7 @@ namespace UE::StylusInput::Private::Wintab
 	DECLARE_DELEGATE(FUpdateWindowRectCallback);
 	DECLARE_DELEGATE(FUpdateTabletContextsCallback);
 
-	class FWintabMessageHandler final : public IWindowsMessageHandler
+	class FWintabMessageHandler : public IWindowsMessageHandler
 	{
 	public:
 		FWintabMessageHandler(IStylusInputInstance* Instance, FGetTabletContextCallback&& GetTabletContextCallback, FGetStylusIDCallback&& GetStylusIDCallback,
@@ -63,7 +63,7 @@ namespace UE::StylusInput::Private::Wintab
 		bool bLastPacketOnDigitizer = false;
 
 #if CHECK_MESSAGE_ORDER_ASSUMPTIONS
-		uint32 CursorChangeTableContextID = 0;
+		uint32 CursorChangeTabletContextID = 0;
 		uint32 CursorChangeSerialNumber = 0;
 #endif
 	};
diff --git a/Engine/Plugins/Editor/StylusInput/Source/StylusInput/Public/StylusInput.h b/Engine/Plugins/Editor/StylusInput/Source/StylusInput/Public/StylusInput.h
--- a/Engine/Plugins/Editor/StylusInput/Source/StylusInput/Public/StylusInput.h
+++ b/Engine/Plugins/Editor/StylusInput/Source/StylusInput/Public/StylusInput.h
@@ -20,11 +20,28 @@ namespace UE::StylusInput
 	 * Any call to this function that returns a valid instance pointer should have a corresponding call to @see #ReleaseInstance.
 	 *
 	 * @param Window The window the stylus input is mapped to.
-	 * @returns The stylus input instance for the given window on the current platform; a nullptr is returned if the instance could not be created or the
-	 * current platform is not supported.
+	 * @returns The stylus input instance for the given window or a nullptr if no valid interface is available or the instance could not be initialized.
 	 */
 	STYLUSINPUT_API IStylusInputInstance* CreateInstance(SWindow& Window);
 
+	/**
+	 * Creates a stylus input instance for a given window and a requested interface.
+	 * Any call to this function that returns a valid instance pointer should have a corresponding call to @see #ReleaseInstance.
+	 *
+	 * @param Window The window the stylus input is mapped to.
+	 * @param Interface Request a specific interface to be used for handling stylus input.
+	 * @param bRequestedInterfaceOnly Only use the requested stylus input interface, and return nullptr if this interface is not available.
+	 * @returns The stylus input instance for the given window or a nullptr if no valid interface is available or the instance could not be initialized.
+	 */
+	STYLUSINPUT_API IStylusInputInstance* CreateInstance(SWindow& Window, FName Interface, bool bRequestedInterfaceOnly);
+
+	/**
+	 * Provides the names of all currently available stylus input interfaces. 
+	 *
+	 * @return Names of available interfaces. 
+	 */
+	STYLUSINPUT_API TArray<FName> GetAvailableInterfaces();
+
 	/**
 	 * Releases any resources held for the given stylus input instance.
 	 * This function should be called for any corresponding successful call to @see #CreateInstance, e.g. when stylus input is no longer needed or the window
@@ -55,7 +72,7 @@ namespace UE::StylusInput
 		/**
 		 * Adds an event handler that gets called for any processed stylus input event.
 		 * Multiple event handlers can get added for a single stylus input instance. Please note that long-running calls will block any other event handlers,
-		 * the processing of additional events, and potentially the game thread itself. Thus make sure to have an event handler call return as soon as possible.
+		 * the processing of additional events, and potentially the game thread itself. Thus, make sure to have an event handler call return as soon as possible.
 		 *
 		 * @param EventHandler The event handler that gets called for each processed stylus input event.
 		 * @param Thread The type of thread the event handler calls are performed on.
@@ -74,7 +91,7 @@ namespace UE::StylusInput
 
 		/**
 		 * Provides the tablet context information for a given tablet context ID.
-		 * This function should only be called with tablet context IDs that were recently provided by the stylus input instance,
+		 * This function should only be called with tablet context IDs recently provided by the stylus input instance,
 		 * e.g. within a @see FStylusPacket as part of an event handler callback.
 		 *
 		 * Calls to this function are thread-safe, i.e. can be made from an event handler callback running on any thread.
@@ -86,7 +103,7 @@ namespace UE::StylusInput
 
 		/**
 		 * Provides the stylus a.k.a. cursor or pen information for a given stylus ID.
-		 * This function should only be called with stylus IDs that were recently provided by the stylus input instance,
+		 * This function should only be called with stylus IDs recently provided by the stylus input instance,
 		 * e.g. within a @see FStylusPacket as part of an event handler callback.
 		 *
 		 * Calls to this function are thread-safe, i.e. can be made from an event handler callback running on any thread.
@@ -105,6 +122,13 @@ namespace UE::StylusInput
 		 */
 		virtual float GetPacketsPerSecond(EEventHandlerThread Thread) const { return -1.0f; }
 
+		/**
+		 * Provides the name of the interface used for this instance. 
+		 * 
+		 * @returns Name of the interface.
+		 */
+		virtual FName GetInterfaceName() = 0;
+
 		/**
 		 * Provides the name of this instance for diagnostic purposes. 
 		 * 
diff --git a/Engine/Plugins/Editor/StylusInput/Source/StylusInput/StylusInput.Build.cs b/Engine/Plugins/Editor/StylusInput/Source/StylusInput/StylusInput.Build.cs
--- a/Engine/Plugins/Editor/StylusInput/Source/StylusInput/StylusInput.Build.cs
+++ b/Engine/Plugins/Editor/StylusInput/Source/StylusInput/StylusInput.Build.cs
@@ -23,7 +23,8 @@ namespace UnrealBuildTool.Rules
 				]
 			);
 
-			const bool bUseWintab = false;
+			const bool bUseWintab = true;
+			const bool bWintabAsDefault = false;
 
 #pragma warning disable CS0162
 			if (Target.Platform == UnrealTargetPlatform.Win64 && bUseWintab)
@@ -33,6 +34,7 @@ namespace UnrealBuildTool.Rules
 #pragma warning restore CS0162
 
 			PublicDefinitions.Add(bUseWintab ? "STYLUSINPUT_USE_WINTAB=1" : "STYLUSINPUT_USE_WINTAB=0");
+			PublicDefinitions.Add(bWintabAsDefault ? "STYLUSINPUT_USE_WINTAB_AS_DEFAULT=1" : "STYLUSINPUT_USE_WINTAB_AS_DEFAULT=0");
 		}
 	}
 }
diff --git a/Engine/Plugins/Editor/StylusInput/Source/StylusInputDebugWidget/Private/StylusInputDebugWidget.cpp b/Engine/Plugins/Editor/StylusInput/Source/StylusInputDebugWidget/Private/StylusInputDebugWidget.cpp
--- a/Engine/Plugins/Editor/StylusInput/Source/StylusInputDebugWidget/Private/StylusInputDebugWidget.cpp
+++ b/Engine/Plugins/Editor/StylusInput/Source/StylusInputDebugWidget/Private/StylusInputDebugWidget.cpp
@@ -40,7 +40,13 @@ namespace UE::StylusInput::DebugWidget
 
 	void SStylusInputDebugWidget::Construct(const FArguments&)
 	{
-		// One-time timer to initialize stylus plugin as soon as widget is up.
+		Interface = []
+		{
+			const TArray<FName> Interfaces = GetAvailableInterfaces();
+			return Interfaces.IsEmpty() ? NAME_None : Interfaces[0];
+		}();
+
+		// One-time timer to initialize stylus plugin as soon as the widget is up.
 		RegisterActiveTimer(
 			0.0f,
 			FWidgetActiveTimerDelegate::CreateLambda([this](double, float)
@@ -90,6 +96,20 @@ namespace UE::StylusInput::DebugWidget
 								SNew(SVerticalBox)
 								+ SVerticalBox::Slot()
 								.AutoHeight()
+								[
+									SNew(SComboButton)
+									.OnGetMenuContent(this, &SStylusInputDebugWidget::GetInterfaceMenu)
+									.ButtonContent()
+									[
+										SNew(STextBlock)
+										.Text_Lambda([&Interface = Interface]
+										{
+											return Interface == NAME_None ? FText::FromString("<not available>") : FText::FromName(Interface);
+										})
+									]
+								]
+								+ SVerticalBox::Slot()
+								.AutoHeight()
 								[
 									SNew(SComboButton)
 									.OnGetMenuContent(this, &SStylusInputDebugWidget::GetEventHandlerThreadMenu)
@@ -493,7 +513,7 @@ namespace UE::StylusInput::DebugWidget
 			return;
 		}
 
-		StylusInput = CreateInstance(*Window);
+		StylusInput = Interface == NAME_None ? CreateInstance(*Window) : CreateInstance(*Window, Interface, true);
 		if (!StylusInput)
 		{
 			LogError("Could not acquire stylus input instance.");
@@ -669,6 +689,46 @@ namespace UE::StylusInput::DebugWidget
 		}
 	}
 
+	TSharedRef<SWidget> SStylusInputDebugWidget::GetInterfaceMenu()
+	{
+		FMenuBuilder MenuBuilder(true, nullptr);
+
+		for (const FName& AvailableInterface : GetAvailableInterfaces())
+		{
+			MenuBuilder.AddMenuEntry(FText::FromName(AvailableInterface),
+			                         {},
+			                         FSlateIcon(),
+			                         FUIAction(FExecuteAction::CreateLambda([this, AvailableInterface]
+			                                   {
+				                                   SetInterface(AvailableInterface);
+			                                   }),
+			                                   FCanExecuteAction(),
+			                                   FIsActionChecked::CreateLambda([&Interface = Interface, AvailableInterface]
+			                                   {
+				                                   return Interface == AvailableInterface;
+			                                   })),
+			                         NAME_None,
+			                         EUserInterfaceActionType::RadioButton);
+		}
+
+		return MenuBuilder.MakeWidget();
+	}
+
+	void SStylusInputDebugWidget::SetInterface(FName InInterface)
+	{
+		if (Interface != InInterface)
+		{
+			UnregisterEventHandler();
+			ReleaseStylusInput();
+
+			Interface = InInterface;
+			LastPacketData = {};
+
+			AcquireStylusInput();
+			RegisterEventHandler();
+		}
+	}
+
 	TSharedRef<SWidget> SStylusInputDebugWidget::GetEventHandlerThreadMenu()
 	{
 		FMenuBuilder MenuBuilder(true, nullptr);
diff --git a/Engine/Plugins/Editor/StylusInput/Source/StylusInputDebugWidget/Private/StylusInputDebugWidget.h b/Engine/Plugins/Editor/StylusInput/Source/StylusInputDebugWidget/Private/StylusInputDebugWidget.h
--- a/Engine/Plugins/Editor/StylusInput/Source/StylusInputDebugWidget/Private/StylusInputDebugWidget.h
+++ b/Engine/Plugins/Editor/StylusInput/Source/StylusInputDebugWidget/Private/StylusInputDebugWidget.h
@@ -74,6 +74,10 @@ namespace UE::StylusInput::DebugWidget
 		void OnPacket(const FStylusInputPacket& Packet);
 		void OnDebugEvent(const FString& Message);
 
+		TSharedRef<SWidget> GetInterfaceMenu();
+		void SetInterface(FName InInterface);
+		FName Interface;
+
 		TSharedRef<SWidget> GetEventHandlerThreadMenu();
 		void SetEventHandlerThread(EEventHandlerThread InEventHandlerThread);
 		EEventHandlerThread EventHandlerThread = EEventHandlerThread::OnGameThread;
