StylusInput: Move interface implementations to separate modules

diff --git a/Engine/Plugins/Editor/StylusInput/Source/StylusInput/Private/StylusInput.cpp b/Engine/Plugins/Editor/StylusInput/Source/StylusInput/Private/StylusInput.cpp
--- a/Engine/Plugins/Editor/StylusInput/Source/StylusInput/Private/StylusInput.cpp
+++ b/Engine/Plugins/Editor/StylusInput/Source/StylusInput/Private/StylusInput.cpp
@@ -1,18 +1,12 @@
 ﻿// Copyright Epic Games, Inc. All Rights Reserved.
 
-#include "StylusInput.h"
+#include <StylusInput.h>
+#include <StylusInputInterface.h>
+#include <StylusInputUtils.h>
 
-#include "StylusInputInterface.h"
-#include "StylusInputUtils.h"
+#include <Algo/AnyOf.h>
 #include <Algo/Find.h>
 
-#if PLATFORM_WINDOWS
-	#if STYLUSINPUT_USE_WINTAB
-		#include "Windows/Wintab/WintabInterface.h"
-	#endif
-	#include "Windows/RealTimeStylus/RealTimeStylusInterface.h"
-#endif
-
 namespace UE::StylusInput
 {
 	namespace Private
@@ -20,44 +14,6 @@ namespace UE::StylusInput
 		class FStylusInputImpl
 		{
 		public:
-			FStylusInputImpl()
-			{
-#if PLATFORM_WINDOWS
-
-	#if STYLUSINPUT_USE_WINTAB && STYLUSINPUT_USE_WINTAB_AS_DEFAULT
-				AddInterface<Wintab::FWintabInterface>();
-	#endif
-
-				AddInterface<RealTimeStylus::FRealTimeStylusInterface>();
-
-	#if STYLUSINPUT_USE_WINTAB && !STYLUSINPUT_USE_WINTAB_AS_DEFAULT
-				AddInterface<Wintab::FWintabInterface>();
-	#endif
-
-#elif PLATFORM_MAC
-
-				// TODO 
-
-#elif PLATFORM_LINUX
-
-				// TODO 
-
-#endif
-			}
-
-			TArray<FName> GetAvailableInterfaces()
-			{
-				TArray<FName> APINames;
-				APINames.Reserve(Interfaces.Num());
-
-				for (const TUniquePtr<IStylusInputInterface>& Interface : Interfaces)
-				{
-					APINames.Add(Interface->GetName());
-				}
-
-				return APINames;
-			}
-
 			IStylusInputInstance* CreateInstance(SWindow& Window, const FName* InterfaceName, const bool bRequestedInterfaceOnly) const
 			{
 				// Helper for trying to create an instance for a given interface.
@@ -75,6 +31,11 @@ namespace UE::StylusInput
 
 						Interface.ReleaseInstance(Instance);
 					}
+					else
+					{
+						UE_LOG(LogStylusInput, Error, TEXT("Failed to create stylus input instance for interface '%s'."),
+							   *Interface.GetName().ToString());
+					}
 
 					return nullptr;
 				};
@@ -132,17 +93,77 @@ namespace UE::StylusInput
 				return false;
 			}
 
-		private:
-			TArray<TUniquePtr<IStylusInputInterface>> Interfaces;
+			TArray<FName> GetAvailableInterfaces()
+			{
+				TArray<FName> APINames;
+				APINames.Reserve(Interfaces.Num());
+
+				for (const TUniquePtr<IStylusInputInterface>& Interface : Interfaces)
+				{
+					APINames.Add(Interface->GetName());
+				}
+
+				return APINames;
+			}
 
-			template <typename InterfaceType>
-			void AddInterface()
+			bool RegisterInterface(IStylusInputInterface* Interface)
 			{
-				if (TUniquePtr<IStylusInputInterface> Interface = InterfaceType::Create())
+				if (!Interface)
 				{
-					Interfaces.Add(MoveTemp(Interface));
+					UE_LOG(Private::LogStylusInput, Warning, TEXT("Nullptr passed into StylusInput::RegisterInterface()."));
+					return false;
 				}
+
+				const FName InterfaceName = Interface->GetName();
+
+				if (Algo::AnyOf(Interfaces, [&InterfaceName](const TUniquePtr<IStylusInputInterface>& I) { return I->GetName() == InterfaceName; }))
+				{
+					UE_LOG(Private::LogStylusInput, Error, TEXT("Interface with name '%s' has already been registered via StylusInput::RegisterInterface()."),
+						   *InterfaceName.ToString())
+					return false;
+				}
+
+				// HACK for temporarily making sure RealTimeStylus is the default under Windows in absence of properly handling this with editor settings and everything.
+				if (InterfaceName == "RealTimeStylus")
+				{
+					Interfaces.EmplaceAt(0, TUniquePtr<IStylusInputInterface>(Interface));
+					return true;
+				}
+
+				Interfaces.Add(TUniquePtr<IStylusInputInterface>(Interface));
+
+				return true;
 			}
+
+			bool UnregisterInterface(IStylusInputInterface* Interface)
+			{
+				if (!Interface)
+				{
+					UE_LOG(Private::LogStylusInput, Warning, TEXT("Nullptr passed into StylusInput::UnregisterInterface()."));
+					return false;
+				}
+
+				const FName InterfaceName = Interface->GetName();
+
+				TUniquePtr<IStylusInputInterface>* RegisteredInterface = Algo::FindByPredicate(
+					Interfaces, [&InterfaceName](const TUniquePtr<IStylusInputInterface>& I) { return I->GetName() == InterfaceName; });
+
+				if (!RegisteredInterface)
+				{
+					UE_LOG(Private::LogStylusInput, Error, TEXT("StylusInput::UnregisterInterface: Interface with name '%s' has not been previously registered."),
+						   *InterfaceName.ToString())
+					return false;
+				}
+
+				[[maybe_unused]] IStylusInputInterface* InterfacePtr = RegisteredInterface->Release(); // TODO workaround for not having to change container type but still not owning the passed in interface. 
+
+				Interfaces.RemoveSingle(*RegisteredInterface);
+
+				return true;
+			}
+
+		private:
+			TArray<TUniquePtr<IStylusInputInterface>> Interfaces;
 		};
 	}
 
@@ -152,11 +173,6 @@ namespace UE::StylusInput
 		return ImplSingleton;
 	}
 
-	TArray<FName> GetAvailableInterfaces()
-	{
-		return GetImplSingleton().GetAvailableInterfaces();
-	}
-
 	IStylusInputInstance* CreateInstance(SWindow& Window)
 	{
 		return GetImplSingleton().CreateInstance(Window, nullptr, false);
@@ -177,4 +193,19 @@ namespace UE::StylusInput
 
 		return GetImplSingleton().ReleaseInstance(Instance);
 	}
+
+	TArray<FName> GetAvailableInterfaces()
+	{
+		return GetImplSingleton().GetAvailableInterfaces();
+	}
+
+	bool RegisterInterface(IStylusInputInterface* Interface)
+	{
+		return GetImplSingleton().RegisterInterface(Interface);
+	}
+
+	bool UnregisterInterface(IStylusInputInterface* Interface)
+	{
+		return GetImplSingleton().UnregisterInterface(Interface);
+	}
 }
diff --git a/Engine/Plugins/Editor/StylusInput/Source/StylusInput/Private/StylusInputInterface.h b/Engine/Plugins/Editor/StylusInput/Source/StylusInput/Private/StylusInputInterface.h
deleted file mode 100644
diff --git a/Engine/Plugins/Editor/StylusInput/Source/StylusInput/Private/StylusInputUtils.h b/Engine/Plugins/Editor/StylusInput/Source/StylusInput/Private/StylusInputUtils.h
deleted file mode 100644
diff --git a/Engine/Plugins/Editor/StylusInput/Source/StylusInput/Private/Windows/RealTimeStylus/RealTimeStylusAPI.cpp b/Engine/Plugins/Editor/StylusInput/Source/StylusInput/Private/Windows/RealTimeStylus/RealTimeStylusAPI.cpp
deleted file mode 100644
diff --git a/Engine/Plugins/Editor/StylusInput/Source/StylusInput/Private/Windows/RealTimeStylus/RealTimeStylusAPI.h b/Engine/Plugins/Editor/StylusInput/Source/StylusInput/Private/Windows/RealTimeStylus/RealTimeStylusAPI.h
deleted file mode 100644
diff --git a/Engine/Plugins/Editor/StylusInput/Source/StylusInput/Private/Windows/RealTimeStylus/RealTimeStylusInstance.cpp b/Engine/Plugins/Editor/StylusInput/Source/StylusInput/Private/Windows/RealTimeStylus/RealTimeStylusInstance.cpp
deleted file mode 100644
diff --git a/Engine/Plugins/Editor/StylusInput/Source/StylusInput/Private/Windows/RealTimeStylus/RealTimeStylusInstance.h b/Engine/Plugins/Editor/StylusInput/Source/StylusInput/Private/Windows/RealTimeStylus/RealTimeStylusInstance.h
deleted file mode 100644
diff --git a/Engine/Plugins/Editor/StylusInput/Source/StylusInput/Private/Windows/RealTimeStylus/RealTimeStylusInterface.cpp b/Engine/Plugins/Editor/StylusInput/Source/StylusInput/Private/Windows/RealTimeStylus/RealTimeStylusInterface.cpp
deleted file mode 100644
diff --git a/Engine/Plugins/Editor/StylusInput/Source/StylusInput/Private/Windows/RealTimeStylus/RealTimeStylusInterface.h b/Engine/Plugins/Editor/StylusInput/Source/StylusInput/Private/Windows/RealTimeStylus/RealTimeStylusInterface.h
deleted file mode 100644
diff --git a/Engine/Plugins/Editor/StylusInput/Source/StylusInput/Private/Windows/RealTimeStylus/RealTimeStylusPluginAsync.cpp b/Engine/Plugins/Editor/StylusInput/Source/StylusInput/Private/Windows/RealTimeStylus/RealTimeStylusPluginAsync.cpp
deleted file mode 100644
diff --git a/Engine/Plugins/Editor/StylusInput/Source/StylusInput/Private/Windows/RealTimeStylus/RealTimeStylusPluginAsync.h b/Engine/Plugins/Editor/StylusInput/Source/StylusInput/Private/Windows/RealTimeStylus/RealTimeStylusPluginAsync.h
deleted file mode 100644
diff --git a/Engine/Plugins/Editor/StylusInput/Source/StylusInput/Private/Windows/RealTimeStylus/RealTimeStylusPluginBase.cpp b/Engine/Plugins/Editor/StylusInput/Source/StylusInput/Private/Windows/RealTimeStylus/RealTimeStylusPluginBase.cpp
deleted file mode 100644
diff --git a/Engine/Plugins/Editor/StylusInput/Source/StylusInput/Private/Windows/RealTimeStylus/RealTimeStylusPluginBase.h b/Engine/Plugins/Editor/StylusInput/Source/StylusInput/Private/Windows/RealTimeStylus/RealTimeStylusPluginBase.h
deleted file mode 100644
diff --git a/Engine/Plugins/Editor/StylusInput/Source/StylusInput/Private/Windows/RealTimeStylus/RealTimeStylusPluginSync.cpp b/Engine/Plugins/Editor/StylusInput/Source/StylusInput/Private/Windows/RealTimeStylus/RealTimeStylusPluginSync.cpp
deleted file mode 100644
diff --git a/Engine/Plugins/Editor/StylusInput/Source/StylusInput/Private/Windows/RealTimeStylus/RealTimeStylusPluginSync.h b/Engine/Plugins/Editor/StylusInput/Source/StylusInput/Private/Windows/RealTimeStylus/RealTimeStylusPluginSync.h
deleted file mode 100644
diff --git a/Engine/Plugins/Editor/StylusInput/Source/StylusInput/Private/Windows/RealTimeStylus/RealTimeStylusStats.h b/Engine/Plugins/Editor/StylusInput/Source/StylusInput/Private/Windows/RealTimeStylus/RealTimeStylusStats.h
deleted file mode 100644
diff --git a/Engine/Plugins/Editor/StylusInput/Source/StylusInput/Private/Windows/RealTimeStylus/RealTimeStylusTabletContext.cpp b/Engine/Plugins/Editor/StylusInput/Source/StylusInput/Private/Windows/RealTimeStylus/RealTimeStylusTabletContext.cpp
deleted file mode 100644
diff --git a/Engine/Plugins/Editor/StylusInput/Source/StylusInput/Private/Windows/RealTimeStylus/RealTimeStylusTabletContext.h b/Engine/Plugins/Editor/StylusInput/Source/StylusInput/Private/Windows/RealTimeStylus/RealTimeStylusTabletContext.h
deleted file mode 100644
diff --git a/Engine/Plugins/Editor/StylusInput/Source/StylusInput/Private/Windows/RealTimeStylus/RealTimeStylusUtils.cpp b/Engine/Plugins/Editor/StylusInput/Source/StylusInput/Private/Windows/RealTimeStylus/RealTimeStylusUtils.cpp
deleted file mode 100644
diff --git a/Engine/Plugins/Editor/StylusInput/Source/StylusInput/Private/Windows/RealTimeStylus/RealTimeStylusUtils.h b/Engine/Plugins/Editor/StylusInput/Source/StylusInput/Private/Windows/RealTimeStylus/RealTimeStylusUtils.h
deleted file mode 100644
diff --git a/Engine/Plugins/Editor/StylusInput/Source/StylusInput/Private/Windows/Wintab/WintabAPI.cpp b/Engine/Plugins/Editor/StylusInput/Source/StylusInput/Private/Windows/Wintab/WintabAPI.cpp
deleted file mode 100644
diff --git a/Engine/Plugins/Editor/StylusInput/Source/StylusInput/Private/Windows/Wintab/WintabAPI.h b/Engine/Plugins/Editor/StylusInput/Source/StylusInput/Private/Windows/Wintab/WintabAPI.h
deleted file mode 100644
diff --git a/Engine/Plugins/Editor/StylusInput/Source/StylusInput/Private/Windows/Wintab/WintabInstance.cpp b/Engine/Plugins/Editor/StylusInput/Source/StylusInput/Private/Windows/Wintab/WintabInstance.cpp
deleted file mode 100644
diff --git a/Engine/Plugins/Editor/StylusInput/Source/StylusInput/Private/Windows/Wintab/WintabInstance.h b/Engine/Plugins/Editor/StylusInput/Source/StylusInput/Private/Windows/Wintab/WintabInstance.h
deleted file mode 100644
diff --git a/Engine/Plugins/Editor/StylusInput/Source/StylusInput/Private/Windows/Wintab/WintabInterface.cpp b/Engine/Plugins/Editor/StylusInput/Source/StylusInput/Private/Windows/Wintab/WintabInterface.cpp
deleted file mode 100644
diff --git a/Engine/Plugins/Editor/StylusInput/Source/StylusInput/Private/Windows/Wintab/WintabInterface.h b/Engine/Plugins/Editor/StylusInput/Source/StylusInput/Private/Windows/Wintab/WintabInterface.h
deleted file mode 100644
diff --git a/Engine/Plugins/Editor/StylusInput/Source/StylusInput/Private/Windows/Wintab/WintabMessageHandler.cpp b/Engine/Plugins/Editor/StylusInput/Source/StylusInput/Private/Windows/Wintab/WintabMessageHandler.cpp
deleted file mode 100644
diff --git a/Engine/Plugins/Editor/StylusInput/Source/StylusInput/Private/Windows/Wintab/WintabMessageHandler.h b/Engine/Plugins/Editor/StylusInput/Source/StylusInput/Private/Windows/Wintab/WintabMessageHandler.h
deleted file mode 100644
diff --git a/Engine/Plugins/Editor/StylusInput/Source/StylusInput/Private/Windows/Wintab/WintabStats.h b/Engine/Plugins/Editor/StylusInput/Source/StylusInput/Private/Windows/Wintab/WintabStats.h
deleted file mode 100644
diff --git a/Engine/Plugins/Editor/StylusInput/Source/StylusInput/Private/Windows/Wintab/WintabStylus.h b/Engine/Plugins/Editor/StylusInput/Source/StylusInput/Private/Windows/Wintab/WintabStylus.h
deleted file mode 100644
diff --git a/Engine/Plugins/Editor/StylusInput/Source/StylusInput/Private/Windows/Wintab/WintabTabletContext.h b/Engine/Plugins/Editor/StylusInput/Source/StylusInput/Private/Windows/Wintab/WintabTabletContext.h
deleted file mode 100644
diff --git a/Engine/Plugins/Editor/StylusInput/Source/StylusInput/Public/StylusInput.h b/Engine/Plugins/Editor/StylusInput/Source/StylusInput/Public/StylusInput.h
--- a/Engine/Plugins/Editor/StylusInput/Source/StylusInput/Public/StylusInput.h
+++ b/Engine/Plugins/Editor/StylusInput/Source/StylusInput/Public/StylusInput.h
@@ -35,13 +35,6 @@ namespace UE::StylusInput
 	 */
 	STYLUSINPUT_API IStylusInputInstance* CreateInstance(SWindow& Window, FName Interface, bool bRequestedInterfaceOnly);
 
-	/**
-	 * Provides the names of all currently available stylus input interfaces. 
-	 *
-	 * @return Names of available interfaces. 
-	 */
-	STYLUSINPUT_API TArray<FName> GetAvailableInterfaces();
-
 	/**
 	 * Releases any resources held for the given stylus input instance.
 	 * This function should be called for any corresponding successful call to @see #CreateInstance, e.g. when stylus input is no longer needed or the window
diff --git a/Engine/Plugins/Editor/StylusInput/Source/StylusInput/Public/StylusInputInterface.h b/Engine/Plugins/Editor/StylusInput/Source/StylusInput/Public/StylusInputInterface.h
new file mode 100644
--- /dev/null
+++ b/Engine/Plugins/Editor/StylusInput/Source/StylusInput/Public/StylusInputInterface.h
@@ -0,0 +1,65 @@
+﻿// Copyright Epic Games, Inc. All Rights Reserved.
+
+#pragma once
+
+#include <UObject/NameTypes.h>
+
+class SWindow;
+
+namespace UE::StylusInput
+{
+	class IStylusInputInstance;
+	class IStylusInputInterface;
+
+	/**
+	 * Provides the names of all currently available stylus input interfaces. 
+	 *
+	 * @return Names of available interfaces. 
+	 */
+	STYLUSINPUT_API TArray<FName> GetAvailableInterfaces();
+
+	/**
+	 * Adds an interface to the StylusInput module.
+	 * 
+	 * Calls to this function are only allowed after the "Default" loading phase. This can be ensured by loading modules providing stylus input interfaces in
+	 * the "PostDefault" loading phase.
+	 *
+	 * @param Interface The interface to be registered.
+	 */
+	STYLUSINPUT_API bool RegisterInterface(IStylusInputInterface* Interface);
+
+	/**
+	 * Removes an interface from the StylusInput module.
+	 * 
+	 * The StylusInput module does not deallocate the provided interface.
+	 * 
+	 * @param Interface The interface to be unregistered.
+	 */
+	STYLUSINPUT_API bool UnregisterInterface(IStylusInputInterface* Interface);
+
+	/**
+	 * Interface for platform APIs that provide stylus input data.
+	 */
+	class IStylusInputInterface
+	{
+	public:
+		virtual ~IStylusInputInterface() = default;
+
+		/**
+		 * Provides the name of the interface; this needs to be unique among all interfaces. 
+		 * 
+		 * @returns Name of the interface.
+		 */
+		virtual FName GetName() = 0;
+
+		/**
+		 * Implementation of /see StylusInput::CreateInstance for this interface. 
+		 */
+		virtual IStylusInputInstance* CreateInstance(SWindow& Window) = 0;
+
+		/**
+		 * Implementation of /see StylusInput::ReleaseInstance for this interface. 
+		 */
+		virtual bool ReleaseInstance(IStylusInputInstance* Instance) = 0;
+	};
+}
diff --git a/Engine/Plugins/Editor/StylusInput/Source/StylusInput/Public/StylusInputTabletContext.h b/Engine/Plugins/Editor/StylusInput/Source/StylusInput/Public/StylusInputTabletContext.h
--- a/Engine/Plugins/Editor/StylusInput/Source/StylusInput/Public/StylusInputTabletContext.h
+++ b/Engine/Plugins/Editor/StylusInput/Source/StylusInput/Public/StylusInputTabletContext.h
@@ -4,7 +4,6 @@
 
 #include <type_traits>
 #include <Containers/UnrealString.h>
-#include <HAL/Platform.h>
 
 namespace UE::StylusInput
 {
diff --git a/Engine/Plugins/Editor/StylusInput/Source/StylusInput/Public/StylusInputUtils.h b/Engine/Plugins/Editor/StylusInput/Source/StylusInput/Public/StylusInputUtils.h
new file mode 100644
--- /dev/null
+++ b/Engine/Plugins/Editor/StylusInput/Source/StylusInput/Public/StylusInputUtils.h
@@ -0,0 +1,127 @@
+﻿// Copyright Epic Games, Inc. All Rights Reserved.
+
+#pragma once
+
+#include <Logging/LogMacros.h>
+#include <Misc/ScopeRWLock.h>
+#include <Templates/SharedPointer.h>
+
+namespace UE::StylusInput::Private
+{
+	STYLUSINPUT_API DECLARE_LOG_CATEGORY_EXTERN(LogStylusInput, Log, All)
+
+	inline void LogError(const FString& Preamble, const FString& Message)
+	{
+		UE_LOG(LogStylusInput, Error, TEXT("%s: %s"), *Preamble, *Message);
+	}
+
+	inline void LogWarning(const FString& Preamble, const FString& Message)
+	{
+		UE_LOG(LogStylusInput, Warning,  TEXT("%s: %s"), *Preamble, *Message);
+	}
+
+	inline void LogVerbose(const FString& Preamble, const FString& Message)
+	{
+		UE_LOG(LogStylusInput, Verbose,  TEXT("%s: %s"), *Preamble, *Message);
+	}
+
+	template <typename DataType, bool ThreadSafe>
+	class TSharedRefDataContainer
+	{
+	public:
+		TSharedRef<DataType> Add(uint32 ID)
+		{
+			WriteScopeLockType WriteLock(RWLock);
+			return Data.Emplace_GetRef(MakeShared<DataType>(ID));
+		}
+
+		void AddOrReplace(TSharedPtr<DataType>&& Item)
+		{
+			if (Item.IsValid())
+			{
+				TSharedRef<DataType>* ExistingItem = Data.FindByPredicate([ID = Item->ID](const TSharedRef<DataType>& Context)
+				{
+					return Context->ID == ID;
+				});
+
+				WriteScopeLockType WriteLock(RWLock);
+				if (ExistingItem != nullptr)
+				{
+					*ExistingItem = MoveTemp(Item).ToSharedRef();
+				}
+				else
+				{
+					Data.Emplace(MoveTemp(Item).ToSharedRef());
+				}
+			}
+		}
+
+		bool Contains(uint32 ID) const
+		{
+			ReadScopeLockType ReadLock(RWLock);
+			const TSharedRef<DataType>* Item = Data.FindByPredicate([ID](const TSharedRef<DataType>& Context)
+			{
+				return Context->ID == ID;
+			});
+			return Item != nullptr;
+		}
+
+		TSharedPtr<DataType> Get(uint32 ID) const
+		{
+			ReadScopeLockType ReadLock(RWLock);
+			const TSharedRef<DataType>* Item = Data.FindByPredicate([ID](const TSharedRef<DataType>& Context)
+			{
+				return Context->ID == ID;
+			});
+			return Item ? Item->ToSharedPtr() : TSharedPtr<DataType>{};
+		}
+
+		bool Remove(uint32 ID)
+		{
+			int32 Index;
+			{
+				ReadScopeLockType ReadLock(RWLock);
+				Index = Data.IndexOfByPredicate([ID](const TSharedRef<DataType>& Tc) { return Tc->ID == ID; });
+				if (Index == INDEX_NONE)
+				{
+					return false;
+				}
+			}
+			WriteScopeLockType WriteLock(RWLock);
+			Data.RemoveAt(Index, EAllowShrinking::No);
+			return true;
+		}
+
+		void Clear()
+		{
+			WriteScopeLockType WriteLock(RWLock);
+			Data.Reset();
+		}
+
+		void Update(const TSharedRefDataContainer<DataType, false>& InData)
+		{
+			WriteScopeLockType WriteLock(RWLock);
+			Data.Reset();
+			for (uint32 Index = 0, Num = InData.Num(); Index < Num; ++Index)
+			{
+				Data.Emplace(InData[Index]);
+			}
+		}
+
+		uint32 Num() const { return Data.Num(); }
+
+		const TSharedRef<DataType>& operator[](const uint32 Index) const { return Data[Index]; }
+
+	private:
+
+		struct FNoLock {};
+		struct FScopeNoLock { explicit FScopeNoLock(FNoLock&) {} };
+
+		using LockType = std::conditional_t<ThreadSafe, FRWLock, FNoLock>;
+		using ReadScopeLockType = std::conditional_t<ThreadSafe, FReadScopeLock, FScopeNoLock>;
+		using WriteScopeLockType = std::conditional_t<ThreadSafe, FWriteScopeLock, FScopeNoLock>;
+
+		UE_NO_UNIQUE_ADDRESS mutable LockType RWLock;
+		TArray<TSharedRef<DataType>> Data;
+	};
+}
diff --git a/Engine/Plugins/Editor/StylusInput/Source/StylusInput/StylusInput.Build.cs b/Engine/Plugins/Editor/StylusInput/Source/StylusInput/StylusInput.Build.cs
--- a/Engine/Plugins/Editor/StylusInput/Source/StylusInput/StylusInput.Build.cs
+++ b/Engine/Plugins/Editor/StylusInput/Source/StylusInput/StylusInput.Build.cs
@@ -22,19 +22,6 @@ namespace UnrealBuildTool.Rules
 					"SlateCore",
 				]
 			);
-
-			const bool bUseWintab = true;
-			const bool bWintabAsDefault = false;
-
-#pragma warning disable CS0162
-			if (Target.Platform == UnrealTargetPlatform.Win64 && bUseWintab)
-			{
-				PrivateDependencyModuleNames.Add("Wintab");
-			}
-#pragma warning restore CS0162
-
-			PublicDefinitions.Add(bUseWintab ? "STYLUSINPUT_USE_WINTAB=1" : "STYLUSINPUT_USE_WINTAB=0");
-			PublicDefinitions.Add(bWintabAsDefault ? "STYLUSINPUT_USE_WINTAB_AS_DEFAULT=1" : "STYLUSINPUT_USE_WINTAB_AS_DEFAULT=0");
 		}
 	}
 }
diff --git a/Engine/Plugins/Editor/StylusInput/Source/StylusInputDebugWidget/Private/StylusInputDebugWidget.cpp b/Engine/Plugins/Editor/StylusInput/Source/StylusInputDebugWidget/Private/StylusInputDebugWidget.cpp
--- a/Engine/Plugins/Editor/StylusInput/Source/StylusInputDebugWidget/Private/StylusInputDebugWidget.cpp
+++ b/Engine/Plugins/Editor/StylusInput/Source/StylusInputDebugWidget/Private/StylusInputDebugWidget.cpp
@@ -3,6 +3,7 @@
 #include "StylusInputDebugWidget.h"
 
 #include <StylusInput.h>
+#include <StylusInputInterface.h>
 #include <StylusInputTabletContext.h>
 #include <Framework/Application/SlateApplication.h>
 #include <Framework/MultiBox/MultiBoxBuilder.h>
diff --git a/Engine/Plugins/Editor/StylusInput/Source/StylusInputDebugWidget/StylusInputDebugWidget.Build.cs b/Engine/Plugins/Editor/StylusInput/Source/StylusInputDebugWidget/StylusInputDebugWidget.Build.cs
--- a/Engine/Plugins/Editor/StylusInput/Source/StylusInputDebugWidget/StylusInputDebugWidget.Build.cs
+++ b/Engine/Plugins/Editor/StylusInput/Source/StylusInputDebugWidget/StylusInputDebugWidget.Build.cs
@@ -7,16 +7,15 @@ namespace UnrealBuildTool.Rules
 		public StylusInputDebugWidget(ReadOnlyTargetRules Target) : base(Target)
 		{
 			PrivateDependencyModuleNames.AddRange(
-				new string[]
-				{
+				[
 					"Core",
 					"Engine",
 					"Slate",
 					"SlateCore",
 					"StylusInput",
 					"UnrealEd",
-					"WorkspaceMenuStructure",
-				}
+					"WorkspaceMenuStructure"
+				]
 			);
 		}
 	}
diff --git a/Engine/Plugins/Editor/StylusInput/Source/StylusInputRealTimeStylus/Private/RealTimeStylusAPI.cpp b/Engine/Plugins/Editor/StylusInput/Source/StylusInputRealTimeStylus/Private/RealTimeStylusAPI.cpp
new file mode 100644
--- /dev/null
+++ b/Engine/Plugins/Editor/StylusInput/Source/StylusInputRealTimeStylus/Private/RealTimeStylusAPI.cpp
@@ -0,0 +1,164 @@
+﻿// Copyright Epic Games, Inc. All Rights Reserved.
+
+#include "RealTimeStylusAPI.h"
+
+#include <KnownFolders.h>
+#include <shtypes.h>
+#include <StylusInputUtils.h>
+#include <Windows/WindowsPlatformProcess.h>
+
+#include "RealTimeStylusUtils.h"
+
+#define LOG_PREAMBLE "RealTimeStylusAPI"
+
+using namespace UE::StylusInput::Private;
+
+namespace UE::StylusInput::RealTimeStylus
+{
+	FRealTimeStylusAPI::FRealTimeStylusAPI()
+	{
+		bInitializedComLibrary = FWindowsPlatformMisc::CoInitialize();
+		if (!bInitializedComLibrary)
+		{
+			LogError(LOG_PREAMBLE, "Could not initialize COM library.");
+			return;
+		}
+
+		auto GetDllHandle = [&DllHandles = DllHandles](const TCHAR* DllName) -> void* {
+			void* DllHandle = FWindowsPlatformProcess::GetDllHandle(DllName);
+			if (DllHandle)
+			{
+				DllHandles.Add(DllHandle);
+				LogVerbose(LOG_PREAMBLE, FString::Format(TEXT("Retrieved DLL handle for {0}."), {DllName}));
+			}
+			else
+			{
+				LogError(LOG_PREAMBLE, FString::Format(TEXT("Could not get DLL handle for {0}."), {DllName}));
+			}
+			return DllHandle;
+		};
+
+		auto GetDllExport = []<typename FFuncType>(const TCHAR* DllName, void* DllHandle, const TCHAR* ExportName, FFuncType& ExportHandle)
+		{
+			void* FuncPtr = FWindowsPlatformProcess::GetDllExport(DllHandle, ExportName);
+			if (FuncPtr)
+			{
+				LogVerbose(LOG_PREAMBLE, FString::Format(TEXT("Retrieved DLL export '{0}' from {1}."), {ExportName, DllName}));
+			}
+			else
+			{
+				LogError(LOG_PREAMBLE, FString::Format(TEXT("Could not get DLL export '{0}' in {1}."), {ExportName, DllName}));
+			}
+			ExportHandle = static_cast<FFuncType>(FuncPtr);
+		};
+
+		const TCHAR* Kernel32DllNAme = TEXT("Kernel32.dll");
+		void* Kernel32DllHandle = GetDllHandle(Kernel32DllNAme);
+		GetDllExport(Kernel32DllNAme, Kernel32DllHandle, TEXT("GetTickCount"), GetTickCount);
+
+		const TCHAR* User32DllName = TEXT("User32.dll");
+		void* User32DllHandle = GetDllHandle(User32DllName);
+		GetDllExport(User32DllName, User32DllHandle, TEXT("GetClientRect"), GetClientRect);
+		GetDllExport(User32DllName, User32DllHandle, TEXT("GetDC"), GetDC);
+		GetDllExport(User32DllName, User32DllHandle, TEXT("ReleaseDC"), ReleaseDC);
+
+		const TCHAR* Gdi32DllName = TEXT("Gdi32.dll");
+		void* Gdi32DllHandle = GetDllHandle(Gdi32DllName);
+		GetDllExport(Gdi32DllName, Gdi32DllHandle, TEXT("GetDeviceCaps"), GetDeviceCaps);
+
+		const TCHAR* Ole32DllName = TEXT("Ole32.dll");
+		void* Ole32DllHandle = GetDllHandle(Ole32DllName);
+		GetDllExport(Ole32DllName, Ole32DllHandle, TEXT("CoCreateInstance"), CoCreateInstance);
+		GetDllExport(Ole32DllName, Ole32DllHandle, TEXT("CoCreateFreeThreadedMarshaler"), CoCreateFreeThreadedMarshaler);
+		GetDllExport(Ole32DllName, Ole32DllHandle, TEXT("CoTaskMemFree"), CoTaskMemFree);
+		GetDllExport(Ole32DllName, Ole32DllHandle, TEXT("StringFromGUID2"), StringFromGUID2);
+
+		const TCHAR* OleAut32DllName = TEXT("OleAut32.dll");
+		void* OleAut32DllHandle = GetDllHandle(OleAut32DllName);
+		GetDllExport(OleAut32DllName, OleAut32DllHandle, TEXT("SysFreeString"), SysFreeString);
+		GetDllExport(OleAut32DllName, OleAut32DllHandle, TEXT("VariantClear"), VariantClear);
+		GetDllExport(OleAut32DllName, OleAut32DllHandle, TEXT("VariantInit"), VariantInit);
+
+		// Try to get handles to RTSCom.dll and InkObj.dll for WindowsInk.
+		{
+			const TCHAR* Shell32DllName = TEXT("Shell32.dll");
+			void* Shell32DllHandle = GetDllHandle(Shell32DllName);
+
+			using FFuncSHGetKnownFolderPath = HRESULT(REFKNOWNFOLDERID, DWORD, HANDLE, PWSTR*);
+			FFuncSHGetKnownFolderPath* SHGetKnownFolderPath = nullptr;
+			GetDllExport(Shell32DllName, Shell32DllHandle, TEXT("SHGetKnownFolderPath"), SHGetKnownFolderPath);
+
+			if (SHGetKnownFolderPath && CoTaskMemFree)
+			{
+				PWSTR ProgramFilesCommonPath = nullptr;
+				if (Succeeded(SHGetKnownFolderPath(FOLDERID_ProgramFilesCommon, 0, nullptr, &ProgramFilesCommonPath), LOG_PREAMBLE))
+				{
+					const FString InkDLLDirectory = FString(ProgramFilesCommonPath) + TEXT("\\microsoft shared\\ink");
+					FPlatformProcess::PushDllDirectory(*InkDLLDirectory);
+					{
+						const TCHAR* RTSComDllName = TEXT("RTSCom.dll");
+						void* RTSComDllHandle = GetDllHandle(RTSComDllName);
+						bHasRTSComDllHandle = RTSComDllHandle != nullptr;
+
+						const TCHAR* InkObjDllName = TEXT("InkObj.dll");
+						void* InkObjDllHandle = GetDllHandle(InkObjDllName);
+						bHasInkObjDllHandle = InkObjDllHandle != nullptr;
+					}
+					FPlatformProcess::PopDllDirectory(*InkDLLDirectory);
+				}
+
+				CoTaskMemFree(ProgramFilesCommonPath);
+			}
+		}
+	}
+
+	FRealTimeStylusAPI::~FRealTimeStylusAPI()
+	{
+		for (void* DllHandle : DllHandles)
+		{
+			FWindowsPlatformProcess::FreeDllHandle(DllHandle);
+		}
+
+		if (bInitializedComLibrary)
+		{
+			FWindowsPlatformMisc::CoUninitialize();
+		}
+	}
+
+	FName FRealTimeStylusAPI::GetName()
+	{
+		static FName Name("RealTimeStylus");
+		return Name;
+	}
+
+	const FRealTimeStylusAPI& FRealTimeStylusAPI::GetInstance()
+	{
+		static FRealTimeStylusAPI API;
+		return API;
+	}
+
+	bool FRealTimeStylusAPI::SatisfiesRequirements() const
+	{
+		bool bSuccess = true;
+
+		bSuccess &= bInitializedComLibrary;
+		bSuccess &= GetTickCount != nullptr;
+		bSuccess &= GetClientRect != nullptr;
+		bSuccess &= GetDC != nullptr;
+		bSuccess &= ReleaseDC != nullptr;
+		bSuccess &= GetDeviceCaps != nullptr;
+		bSuccess &= CoCreateInstance != nullptr;
+		bSuccess &= CoCreateFreeThreadedMarshaler != nullptr;
+		bSuccess &= CoTaskMemFree != nullptr;
+		bSuccess &= StringFromGUID2 != nullptr;
+		bSuccess &= SysFreeString != nullptr;
+		bSuccess &= VariantClear != nullptr;
+		bSuccess &= VariantInit != nullptr;
+		bSuccess &= bHasRTSComDllHandle;
+		bSuccess &= bHasInkObjDllHandle;
+
+		return bSuccess;
+	}
+}
+
+#undef LOG_PREAMBLE
diff --git a/Engine/Plugins/Editor/StylusInput/Source/StylusInputRealTimeStylus/Private/RealTimeStylusAPI.h b/Engine/Plugins/Editor/StylusInput/Source/StylusInputRealTimeStylus/Private/RealTimeStylusAPI.h
new file mode 100644
--- /dev/null
+++ b/Engine/Plugins/Editor/StylusInput/Source/StylusInputRealTimeStylus/Private/RealTimeStylusAPI.h
@@ -0,0 +1,71 @@
+﻿// Copyright Epic Games, Inc. All Rights Reserved.
+
+#pragma once
+
+#include <Containers/Array.h>
+
+#include <Windows/AllowWindowsPlatformTypes.h>
+	#include <windef.h>
+	#include <wtypes.h>
+	#include <Unknwn.h>
+#include <Windows/HideWindowsPlatformTypes.h>
+
+namespace UE::StylusInput::RealTimeStylus
+{
+	class FRealTimeStylusAPI
+	{
+	public:
+		static FName GetName();
+		static const FRealTimeStylusAPI& GetInstance();
+
+		bool SatisfiesRequirements() const;
+
+		using FFuncGetTickCount = DWORD();
+		FFuncGetTickCount* GetTickCount = nullptr;
+
+		using FFuncGetClientRect = BOOL(HWND, LPRECT);
+		FFuncGetClientRect* GetClientRect = nullptr;
+
+		using FFuncGetDC = HDC(HWND);
+		FFuncGetDC* GetDC = nullptr;
+
+		using FFuncReleaseDC = int(HWND, HDC);
+		FFuncReleaseDC* ReleaseDC = nullptr;
+
+		using FFuncGetDeviceCaps = int(HDC, int);
+		FFuncGetDeviceCaps* GetDeviceCaps = nullptr;
+
+		using FFuncCoCreateInstance = HRESULT(REFCLSID, IUnknown*, DWORD, REFIID, LPVOID);
+		FFuncCoCreateInstance* CoCreateInstance = nullptr;
+
+		using FFuncCoCreateFreeThreadedMarshaler = HRESULT(LPUNKNOWN, LPUNKNOWN*);
+		FFuncCoCreateFreeThreadedMarshaler* CoCreateFreeThreadedMarshaler = nullptr;
+
+		using FFuncCoTaskMemFree = void(LPVOID);
+		FFuncCoTaskMemFree* CoTaskMemFree = nullptr;
+
+		using FFuncStringFromGUID2 = int(REFGUID, LPOLESTR, int);
+		FFuncStringFromGUID2* StringFromGUID2 = nullptr;
+
+		using FFuncSysFreeString = void(BSTR);
+		FFuncSysFreeString* SysFreeString = nullptr;
+
+		using FFuncVariantClear = HRESULT(VARIANTARG*);
+		FFuncVariantClear* VariantClear = nullptr;
+
+		using FFuncVariantInit = void(VARIANTARG*);
+		FFuncVariantInit* VariantInit = nullptr;
+
+	private:
+
+		FRealTimeStylusAPI();
+		~FRealTimeStylusAPI();
+
+		TArray<void*> DllHandles;
+
+		bool bInitializedComLibrary = false;
+		bool bHasRTSComDllHandle = false;
+		bool bHasInkObjDllHandle = false;
+
+	};
+}
diff --git a/Engine/Plugins/Editor/StylusInput/Source/StylusInputRealTimeStylus/Private/RealTimeStylusInstance.cpp b/Engine/Plugins/Editor/StylusInput/Source/StylusInputRealTimeStylus/Private/RealTimeStylusInstance.cpp
new file mode 100644
--- /dev/null
+++ b/Engine/Plugins/Editor/StylusInput/Source/StylusInputRealTimeStylus/Private/RealTimeStylusInstance.cpp
@@ -0,0 +1,507 @@
+﻿// Copyright Epic Games, Inc. All Rights Reserved.
+
+#include "RealTimeStylusInstance.h"
+
+#include <StylusInput.h>
+#include <StylusInputUtils.h>
+#include <Algo/Transform.h>
+
+#include "RealTimeStylusAPI.h"
+#include "RealTimeStylusPluginAsync.h"
+#include "RealTimeStylusPluginSync.h"
+#include "RealTimeStylusUtils.h"
+
+// Fixed indices at which the plugins are installed; we only ever have one each.
+#define ASYNC_PLUGIN_INDEX 0
+#define SYNC_PLUGIN_INDEX 0
+
+#define LOCTEXT_NAMESPACE "RealTimeStylusInstance"
+#define LOG_PREAMBLE "RealTimeStylusInstance"
+
+using namespace UE::StylusInput::Private;
+
+namespace UE::StylusInput::RealTimeStylus
+{
+	FRealTimeStylusInstance::FRealTimeStylusInstance(const uint32 ID, const HWND OSWindowHandle)
+		: ID(ID)
+		, RealTimeStylusAPI(FRealTimeStylusAPI::GetInstance())
+	{
+		SetupWindowContext(OSWindowHandle);
+
+		void* OutInstance = nullptr;
+		if (Failed(RealTimeStylusAPI.CoCreateInstance(CLSID_RealTimeStylus, nullptr, CLSCTX_ALL, IID_IRealTimeStylus, &OutInstance), LOG_PREAMBLE))
+		{
+			LogError(LOG_PREAMBLE, "Could not create RealTimeStylus COM object instance!");
+		}
+		RealTimeStylus = static_cast<IRealTimeStylus*>(OutInstance);
+
+		if (RealTimeStylus.IsValid())
+		{
+			TArray<GUID> DesiredPackets;
+			DesiredPackets.Reserve(std::size(PacketPropertyConstants));
+			Algo::Transform(PacketPropertyConstants, DesiredPackets, [](const FPacketPropertyConstant& Property) { return Property.Guid; });
+
+			if (Failed(RealTimeStylus->SetDesiredPacketDescription(DesiredPackets.Num(), DesiredPackets.GetData()), LOG_PREAMBLE))
+			{
+				LogError(LOG_PREAMBLE, "Could not set desired RealTimeStylus packet description!");
+				bWasInitializedSuccessfully = false;
+			}
+
+			if (Failed(RealTimeStylus->put_HWND(reinterpret_cast<HANDLE_PTR>(OSWindowHandle)), LOG_PREAMBLE))
+			{
+				LogError(LOG_PREAMBLE, "Could not set RealTimeStylus window handle!");
+				bWasInitializedSuccessfully = false;
+			}
+
+			if (Failed(RealTimeStylus->put_Enabled(::Windows::TRUE), LOG_PREAMBLE))
+			{
+				LogError(LOG_PREAMBLE, "Could not enable RealTimeStylus interface!");
+				bWasInitializedSuccessfully = false;
+			}
+		}
+		else
+		{
+			LogError(LOG_PREAMBLE, "Stylus input via the RealTimeStylus interface is not available!");
+			bWasInitializedSuccessfully = false;
+		}
+
+		if (bWasInitializedSuccessfully)
+		{
+			LogVerbose(LOG_PREAMBLE, "Successfully initialized RealTimeStylus COM object.");
+		}
+	}
+
+	FRealTimeStylusInstance::~FRealTimeStylusInstance()
+	{
+		if (RealTimeStylus.IsValid())
+		{
+			RealTimeStylus->RemoveAllStylusAsyncPlugins();
+			RealTimeStylus->RemoveAllStylusSyncPlugins();
+			RealTimeStylus->Release();
+		}
+	}
+
+	bool FRealTimeStylusInstance::AddEventHandler(IStylusInputEventHandler* EventHandler, const EEventHandlerThread Thread)
+	{
+		if (!RealTimeStylus.IsValid())
+		{
+			LogWarning(LOG_PREAMBLE, "Stylus input via the RealTimeStylus interface is not available; event handler was not installed.");
+			return false;
+		}
+
+		if (!EventHandler)
+		{
+			LogWarning(LOG_PREAMBLE, "Tried to add nullptr as event handler.");
+			return false;
+		}
+
+		if (Thread == EEventHandlerThread::OnGameThread)
+		{
+			if (!AsyncPlugin)
+			{
+				EnablePlugin(EEventHandlerThread::OnGameThread, EventHandler);
+
+				if (!AsyncPlugin)
+				{
+					LogError(LOG_PREAMBLE, FString::Format(
+								 TEXT("Event handler '{0}' was not added since the game thread stylus input plugin could not be installed."),
+								 {EventHandler->GetName()}));
+					return false;
+				}
+
+				return true;
+			}
+
+			return AsyncPlugin->AddEventHandler(EventHandler);
+		}
+
+		if (Thread == EEventHandlerThread::Asynchronous)
+		{
+			if (!SyncPlugin)
+			{
+				EnablePlugin(EEventHandlerThread::Asynchronous, EventHandler);
+
+				if (!SyncPlugin)
+				{
+					LogError(LOG_PREAMBLE, FString::Format(
+								 TEXT("Event handler '{0}' was not added since the asynchronous stylus input plugin could not be installed."),
+								 {EventHandler->GetName()}));
+					return false;
+				}
+
+				return true;
+			}
+
+			return SyncPlugin->AddEventHandler(EventHandler);
+		}
+
+		return false;
+	}
+
+	bool FRealTimeStylusInstance::RemoveEventHandler(IStylusInputEventHandler* EventHandler)
+	{
+		if (!RealTimeStylus.IsValid())
+		{
+			return false;
+		}
+
+		if (!EventHandler)
+		{
+			LogWarning(LOG_PREAMBLE, "Tried to remove nullptr event handler.");
+			return false;
+		}
+
+		bool bWasRemoved = false;
+
+		if (AsyncPlugin && AsyncPlugin->RemoveEventHandler(EventHandler))
+		{
+			bWasRemoved = true;
+
+			if (AsyncPlugin->NumEventHandlers() == 0)
+			{
+				DisablePlugin(EEventHandlerThread::OnGameThread);
+			}
+		}
+
+		if (SyncPlugin && SyncPlugin->RemoveEventHandler(EventHandler))
+		{
+			bWasRemoved = true;
+
+			if (SyncPlugin->NumEventHandlers() == 0)
+			{
+				DisablePlugin(EEventHandlerThread::Asynchronous);
+			}
+		}
+
+		if (!bWasRemoved)
+		{
+			LogError(LOG_PREAMBLE, FString::Format(TEXT("Event handler {0} does not exist."), {EventHandler->GetName()}));
+		}
+
+		return bWasRemoved;
+	}
+
+	const TSharedPtr<IStylusInputTabletContext> FRealTimeStylusInstance::GetTabletContext(const uint32 TabletContextID)
+	{
+		return TabletContexts.Get(TabletContextID);
+	}
+
+	const TSharedPtr<IStylusInputStylusInfo> FRealTimeStylusInstance::GetStylusInfo(uint32 StylusID)
+	{
+		return StylusInfos.Get(StylusID);
+	}
+
+	float FRealTimeStylusInstance::GetPacketsPerSecond(const EEventHandlerThread EventHandlerThread) const
+	{
+		if (EventHandlerThread == EEventHandlerThread::Asynchronous)
+		{
+			return SyncPlugin ? SyncPlugin->GetPacketsPerSecond() : 0.0f;
+		}
+
+		if (EventHandlerThread == EEventHandlerThread::OnGameThread)
+		{
+			return AsyncPlugin ? AsyncPlugin->GetPacketsPerSecond() : 0.0f;
+		}
+
+		return -1.0f;
+	}
+
+	FName FRealTimeStylusInstance::GetInterfaceName()
+	{
+		return FRealTimeStylusAPI::GetName();
+	}
+
+	FText FRealTimeStylusInstance::GetName()
+	{
+		return FText::Format(LOCTEXT("RealTimeStylus", "RealTimeStylus #{0}"), ID);
+	}
+
+	bool FRealTimeStylusInstance::WasInitializedSuccessfully()
+	{
+		return bWasInitializedSuccessfully;
+	}
+
+	void FRealTimeStylusInstance::SetupWindowContext(const HWND HWindow)
+	{
+		RECT WindowRectangle;
+		if (!ensure(RealTimeStylusAPI.GetClientRect(HWindow, &WindowRectangle)))
+		{
+			LogError(LOG_PREAMBLE, "Could not retrieve window rectangle; failed to setup device context!");
+			return;
+		}
+
+		if (WindowRectangle.right <= WindowRectangle.left || WindowRectangle.bottom <= WindowRectangle.top)
+		{
+			LogError(LOG_PREAMBLE, "Window appears to be minimized; failed to setup device context!");
+			return;
+		}
+
+		const HDC HWindowDeviceContext = RealTimeStylusAPI.GetDC(HWindow);
+		if (!ensure(HWindowDeviceContext))
+		{
+			LogError(LOG_PREAMBLE, "Could not retrieve window device context; failed to setup device context!");
+			return;
+		}
+
+		const FVector2d DisplaySizeInMillimeters = {
+			static_cast<double>(RealTimeStylusAPI.GetDeviceCaps(HWindowDeviceContext, HORZSIZE)),
+			static_cast<double>(RealTimeStylusAPI.GetDeviceCaps(HWindowDeviceContext, VERTSIZE))
+		};
+
+		if (!ensure(DisplaySizeInMillimeters.X > 0.0 && DisplaySizeInMillimeters.Y > 0.0))
+		{
+			LogError(LOG_PREAMBLE, "Display size in millimeters is invalid; failed to setup device context!");
+			return;
+		}
+
+		const FVector2d DisplaySizeInPixels = {
+			static_cast<double>(RealTimeStylusAPI.GetDeviceCaps(HWindowDeviceContext, HORZRES)),
+			static_cast<double>(RealTimeStylusAPI.GetDeviceCaps(HWindowDeviceContext, VERTRES))
+		};
+
+		if (!ensure(DisplaySizeInPixels.X > 0.0 && DisplaySizeInPixels.Y > 0.0))
+		{
+			LogError(LOG_PREAMBLE, "Display size in pixels is invalid; failed to setup device context!");
+			return;
+		}
+
+		const FVector2d DisplayPixelsPerLogicalInch = {
+			static_cast<double>(RealTimeStylusAPI.GetDeviceCaps(HWindowDeviceContext, LOGPIXELSX)),
+			static_cast<double>(RealTimeStylusAPI.GetDeviceCaps(HWindowDeviceContext, LOGPIXELSY))
+		};
+
+		if (!ensure(DisplayPixelsPerLogicalInch.X > 0.0 && DisplayPixelsPerLogicalInch.Y > 0.0))
+		{
+			LogError(LOG_PREAMBLE, "Display pixels per logical inch is invalid; failed to setup device context!");
+			return;
+		}
+
+		const FVector2d WindowSizeInPixels = {
+			static_cast<double>(WindowRectangle.right - WindowRectangle.left),
+			static_cast<double>(WindowRectangle.bottom - WindowRectangle.top)
+		};
+
+		const FVector2d DisplayMillimetersPerPixel = DisplaySizeInMillimeters / DisplaySizeInPixels;
+		const FVector2d WindowSizeInMillimeters = WindowSizeInPixels * DisplayMillimetersPerPixel;
+		const FVector2d MillimetersPerInch{25.4};
+		const FVector2d DisplayInchesPerPixel = MillimetersPerInch / DisplayMillimetersPerPixel;
+
+		WindowContext.XYScale = DisplayInchesPerPixel / DisplayPixelsPerLogicalInch;
+		WindowContext.XYMaximum = WindowSizeInMillimeters * FVector2d{100.0} + FVector2d{0.5};
+		WindowContext.WindowSize = WindowSizeInPixels;
+
+		RealTimeStylusAPI.ReleaseDC(HWindow, HWindowDeviceContext);
+
+		LogVerbose(LOG_PREAMBLE, "Successfully setup window context.");
+	}
+
+	const FWindowContext& FRealTimeStylusInstance::GetWindowContext() const
+	{
+		return WindowContext;
+	}
+
+	void FRealTimeStylusInstance::UpdateStylusInfo(STYLUS_ID StylusID)
+	{
+		if (!ensure(RealTimeStylus.IsValid()))
+		{
+			return;
+		}
+
+		TSharedPtr<FStylusInfo> StylusInfo = MakeShared<FStylusInfo>(StylusID);
+
+		IInkCursor* InkCursor = nullptr;
+		if (Succeeded(RealTimeStylus->GetStylusForId(StylusID, &InkCursor), LOG_PREAMBLE) && InkCursor != nullptr)
+		{
+			BSTR String = nullptr;
+			if (Succeeded(InkCursor->get_Name(&String), LOG_PREAMBLE))
+			{
+				StylusInfo->Name = String;
+				RealTimeStylusAPI.SysFreeString(String);
+			}
+			else
+			{
+				LogError(LOG_PREAMBLE, FString::Format(TEXT("Could not get name for stylus with ID {0}."), {static_cast<uint32>(StylusID)}));
+			}
+
+			IInkCursorButtons* Buttons = nullptr;
+			long NumButtons = 0;
+			if (Succeeded(InkCursor->get_Buttons(&Buttons), LOG_PREAMBLE) &&
+				Buttons && Succeeded(Buttons->get_Count(&NumButtons), LOG_PREAMBLE))
+			{
+				if (NumButtons > 0)
+				{
+					TArray<FStylusButton> StylusInfoButtons;
+					StylusInfoButtons.Reserve(NumButtons);
+
+					for (long ButtonIndex = 0; ButtonIndex < NumButtons; ++ButtonIndex)
+					{
+						VARIANT ButtonIdentifier;
+						RealTimeStylusAPI.VariantInit(&ButtonIdentifier);
+
+						ButtonIdentifier.vt = VT_INT;
+						ButtonIdentifier.lVal = ButtonIndex;
+
+						IInkCursorButton* Button = nullptr;
+						if (Succeeded(Buttons->Item(ButtonIdentifier, &Button), LOG_PREAMBLE))
+						{
+							FStylusButton NewButton;
+
+							if (Succeeded(Button->get_Id(&String), LOG_PREAMBLE))
+							{
+								NewButton.ID = String;
+								RealTimeStylusAPI.SysFreeString(String);
+							}
+							else
+							{
+								LogError(LOG_PREAMBLE, FString::Format(
+											 TEXT("Could not get ID for button {0} for stylus with ID {1}."), {static_cast<int32>(ButtonIndex), static_cast<uint32>(StylusID)}));
+							}
+
+							if (Succeeded(Button->get_Name(&String), LOG_PREAMBLE))
+							{
+								NewButton.Name = String;
+								RealTimeStylusAPI.SysFreeString(String);
+							}
+							else
+							{
+								LogError(LOG_PREAMBLE, FString::Format(
+											 TEXT("Could not get name for button {0} for stylus with ID {1}."), {static_cast<int32>(ButtonIndex), static_cast<uint32>(StylusID)}));
+							}
+
+							if (!NewButton.ID.IsEmpty() && !NewButton.Name.IsEmpty())
+							{
+								StylusInfoButtons.Emplace(NewButton);
+							}
+						}
+						else
+						{
+							LogError(LOG_PREAMBLE, FString::Format(
+										 TEXT("Could not get button {0} for stylus with ID {1}."), {static_cast<int32>(ButtonIndex), static_cast<uint32>(StylusID)}));
+						}
+
+						RealTimeStylusAPI.VariantClear(&ButtonIdentifier);
+					}
+
+					StylusInfo->Buttons = MoveTemp(StylusInfoButtons);
+				}
+			}
+			else
+			{
+				LogError(LOG_PREAMBLE, FString::Format(TEXT("Could not get buttons for stylus with ID {0}."), {static_cast<uint32>(StylusID)}));
+			}
+		}
+		else
+		{
+			LogWarning(LOG_PREAMBLE, FString::Format(TEXT("Could not get stylus info for ID {0}."), {static_cast<uint32>(StylusID)}));
+		}
+
+		StylusInfos.AddOrReplace(MoveTemp(StylusInfo));
+	}
+
+	void FRealTimeStylusInstance::UpdateTabletContexts(const FTabletContextContainer& InTabletContexts)
+	{
+		TabletContexts.Update(InTabletContexts);
+	}
+
+	void FRealTimeStylusInstance::EnablePlugin(const EEventHandlerThread EventHandlerThread, IStylusInputEventHandler* EventHandler)
+	{
+		if (!ensure(RealTimeStylus.IsValid()))
+		{
+			return;
+		}
+
+		if (EventHandlerThread == EEventHandlerThread::OnGameThread)
+		{
+			AsyncPlugin = MakeUnique<FRealTimeStylusPluginAsync>(
+				this,
+				FGetWindowContextCallback::CreateRaw(this, &FRealTimeStylusInstance::GetWindowContext),
+				FUpdateTabletContextsCallback::CreateRaw(this, &FRealTimeStylusInstance::UpdateTabletContexts),
+				FUpdateStylusInfoCallback::CreateRaw(this, &FRealTimeStylusInstance::UpdateStylusInfo),
+				EventHandler
+			);
+
+			if (Succeeded(RealTimeStylus->AddStylusAsyncPlugin(ASYNC_PLUGIN_INDEX, AsyncPlugin.Get()), LOG_PREAMBLE))
+			{
+				LogVerbose(LOG_PREAMBLE, "Added game thread stylus input plugin!");
+			}
+			else
+			{
+				LogError(LOG_PREAMBLE, "Could not add game thread stylus input plugin!");
+				AsyncPlugin.Reset();
+			}
+		}
+
+		if (EventHandlerThread == EEventHandlerThread::Asynchronous)
+		{
+			SyncPlugin = MakeUnique<FRealTimeStylusPluginSync>(
+				this,
+				FGetWindowContextCallback::CreateRaw(this, &FRealTimeStylusInstance::GetWindowContext),
+				FUpdateTabletContextsCallback::CreateRaw(this, &FRealTimeStylusInstance::UpdateTabletContexts),
+				FUpdateStylusInfoCallback::CreateRaw(this, &FRealTimeStylusInstance::UpdateStylusInfo),
+				EventHandler
+			);
+
+			if (Succeeded(SyncPlugin->CreateFreeThreadMarshaler(), LOG_PREAMBLE))
+			{
+				if (Succeeded(RealTimeStylus->AddStylusSyncPlugin(SYNC_PLUGIN_INDEX, SyncPlugin.Get()), LOG_PREAMBLE))
+				{
+					LogVerbose(LOG_PREAMBLE, "Added asynchronous stylus input plugin!");
+				}
+				else
+				{
+					LogError(LOG_PREAMBLE, "Could not add asynchronous stylus input plugin!");
+					SyncPlugin.Reset();
+				}
+			}
+			else
+			{
+				LogError(LOG_PREAMBLE, "Could not create free thread marshaler for asynchronous stylus input plugin!");
+				SyncPlugin.Reset();
+			}
+		}
+	}
+
+	void FRealTimeStylusInstance::DisablePlugin(const EEventHandlerThread EventHandlerThread)
+	{
+		if (!ensure(RealTimeStylus.IsValid()))
+		{
+			return;
+		}
+
+#pragma warning(push)
+#pragma warning(disable : 6387)
+
+		if (EventHandlerThread == EEventHandlerThread::OnGameThread)
+		{
+			if (Failed(RealTimeStylus->RemoveStylusAsyncPlugin(ASYNC_PLUGIN_INDEX, nullptr), LOG_PREAMBLE))
+			{
+				LogError(LOG_PREAMBLE, "Could not remove game thread stylus input plugin!");
+			}
+			else
+			{
+				LogVerbose(LOG_PREAMBLE, "Removed game thread stylus input plugin!");
+			}
+
+			AsyncPlugin.Reset();
+		}
+
+		if (EventHandlerThread == EEventHandlerThread::Asynchronous)
+		{
+			if (Failed(RealTimeStylus->RemoveStylusSyncPlugin(SYNC_PLUGIN_INDEX, nullptr), LOG_PREAMBLE))
+			{
+				LogError(LOG_PREAMBLE, "Could not remove asynchronous stylus input plugin!");
+			}
+			else
+			{
+				LogVerbose(LOG_PREAMBLE, "Removed asynchronous stylus input plugin!");
+			}
+
+			SyncPlugin.Reset();
+		}
+
+#pragma warning(pop)
+	}
+}
+
+#undef LOG_PREAMBLE
+#undef LOCTEXT_NAMESPACE
diff --git a/Engine/Plugins/Editor/StylusInput/Source/StylusInputRealTimeStylus/Private/RealTimeStylusInstance.h b/Engine/Plugins/Editor/StylusInput/Source/StylusInputRealTimeStylus/Private/RealTimeStylusInstance.h
new file mode 100644
--- /dev/null
+++ b/Engine/Plugins/Editor/StylusInput/Source/StylusInputRealTimeStylus/Private/RealTimeStylusInstance.h
@@ -0,0 +1,68 @@
+﻿// Copyright Epic Games, Inc. All Rights Reserved.
+
+#pragma once
+
+#include <StylusInput.h>
+#include <Templates/UniquePtr.h>
+
+#include <Windows/AllowWindowsPlatformTypes.h>
+	#include <Microsoft/COMPointer.h>
+	#include <RTSCom.h>
+#include <Windows/HideWindowsPlatformTypes.h>
+
+#include "RealTimeStylusTabletContext.h"
+
+namespace UE::StylusInput
+{
+	class IStylusInputEventHandler;
+}
+
+namespace UE::StylusInput::RealTimeStylus
+{
+	class FRealTimeStylusPluginSync;
+	class FRealTimeStylusPluginAsync;
+	class FRealTimeStylusAPI;
+
+	class FRealTimeStylusInstance : public IStylusInputInstance
+	{
+	public:
+		explicit FRealTimeStylusInstance(uint32 ID, HWND OSWindowHandle);
+		virtual ~FRealTimeStylusInstance() override;
+
+		virtual bool AddEventHandler(IStylusInputEventHandler* EventHandler, EEventHandlerThread Thread) override;
+		virtual bool RemoveEventHandler(IStylusInputEventHandler* EventHandler) override;
+
+		virtual const TSharedPtr<IStylusInputTabletContext> GetTabletContext(uint32 TabletContextID) override;
+		virtual const TSharedPtr<IStylusInputStylusInfo> GetStylusInfo(uint32 StylusID) override;
+
+		virtual float GetPacketsPerSecond(EEventHandlerThread EventHandlerThread) const override;
+
+		virtual FName GetInterfaceName() override;
+		virtual FText GetName() override;
+
+		virtual bool WasInitializedSuccessfully() override;
+
+	private:
+		void EnablePlugin(EEventHandlerThread EventHandlerThread, IStylusInputEventHandler* EventHandler);
+		void DisablePlugin(EEventHandlerThread EventHandlerThread);
+
+		void SetupWindowContext(HWND HWindow);
+		const FWindowContext& GetWindowContext() const;
+
+		void UpdateStylusInfo(STYLUS_ID StylusID);
+		void UpdateTabletContexts(const FTabletContextContainer& InTabletContexts);
+
+		const uint32 ID;
+		const FRealTimeStylusAPI& RealTimeStylusAPI;
+		FWindowContext WindowContext;
+
+		TComPtr<IRealTimeStylus> RealTimeStylus;
+		TUniquePtr<FRealTimeStylusPluginAsync> AsyncPlugin;
+		TUniquePtr<FRealTimeStylusPluginSync> SyncPlugin;
+
+		FTabletContextThreadSafeContainer TabletContexts;
+		FStylusInfoThreadSafeContainer StylusInfos;
+
+		bool bWasInitializedSuccessfully = true;
+	};
+}
diff --git a/Engine/Plugins/Editor/StylusInput/Source/StylusInputRealTimeStylus/Private/RealTimeStylusInterface.cpp b/Engine/Plugins/Editor/StylusInput/Source/StylusInputRealTimeStylus/Private/RealTimeStylusInterface.cpp
new file mode 100644
--- /dev/null
+++ b/Engine/Plugins/Editor/StylusInput/Source/StylusInputRealTimeStylus/Private/RealTimeStylusInterface.cpp
@@ -0,0 +1,98 @@
+﻿// Copyright Epic Games, Inc. All Rights Reserved.
+
+#include "RealTimeStylusInterface.h"
+
+#include <StylusInput.h>
+#include <StylusInputUtils.h>
+
+#include <GenericPlatform/GenericWindow.h>
+#include <Templates/SharedPointer.h>
+#include <Widgets/SWindow.h>
+
+#include <Windows/AllowWindowsPlatformTypes.h>
+	#include <RTSCOM_i.c>
+#include <Windows/HideWindowsPlatformTypes.h>
+
+#include "RealTimeStylusInstance.h"
+#include "RealTimeStylusAPI.h"
+
+using namespace UE::StylusInput::Private;
+
+namespace UE::StylusInput::RealTimeStylus
+{
+	FName FRealTimeStylusInterface::GetName()
+	{
+		return FRealTimeStylusAPI::GetName();
+	}
+
+	IStylusInputInstance* FRealTimeStylusInterface::CreateInstance(SWindow& Window)
+	{
+		if (FRefCountedInstance* ExistingRefCountedInstance = Instances.Find(&Window))
+		{
+			++ExistingRefCountedInstance->RefCount;
+			return ExistingRefCountedInstance->Instance.Get();
+		}
+
+		HWND OSWindowHandle = [&Window]
+		{
+			const TSharedPtr<const FGenericWindow> NativeWindow = Window.GetNativeWindow();
+			return NativeWindow.IsValid() ? static_cast<HWND>(NativeWindow->GetOSWindowHandle()) : nullptr;
+		}();
+
+		if (!OSWindowHandle)
+		{
+			LogError("RealTimeStylusInterface", "Could not get native window handle.");
+			return nullptr;
+		}
+
+		FRealTimeStylusInstance* NewInstance = Instances.Emplace(
+			&Window, {MakeUnique<FRealTimeStylusInstance>(NextInstanceID++, OSWindowHandle), 1}).Instance.Get();
+		if (!ensureMsgf(NewInstance, TEXT("RealTimeStylusInterface: Failed to create stylus input instance.")))
+		{
+			Instances.Remove(&Window);
+			return nullptr;
+		}
+
+		return NewInstance;
+	}
+
+	bool FRealTimeStylusInterface::ReleaseInstance(IStylusInputInstance* Instance)
+	{
+		check(Instance);
+
+		FRealTimeStylusInstance *const RealTimeStylusInstance = static_cast<FRealTimeStylusInstance*>(Instance);
+
+		// Find existing instance
+		for (TTuple<SWindow*, FRefCountedInstance>& Entry : Instances)
+		{
+			FRefCountedInstance& RefCountedInstance = Entry.Get<1>();
+
+			if (RefCountedInstance.Instance.Get() == RealTimeStylusInstance)
+			{
+				// Decrease reference count
+				check(RefCountedInstance.RefCount > 0);
+				if (--RefCountedInstance.RefCount == 0)
+				{
+					// Delete if there are no references left
+					Instances.Remove(Entry.Key);
+				}
+
+				return true;
+			}
+		}
+
+		ensureMsgf(false, TEXT("RealTimeStylusInterface: Failed to find provided instance."));
+		return false;
+	}
+
+	TUniquePtr<IStylusInputInterface> FRealTimeStylusInterface::Create()
+	{
+		if (FRealTimeStylusAPI::GetInstance().SatisfiesRequirements())
+		{
+			TUniquePtr<FRealTimeStylusInterface> RealTimeStylusInterface = MakeUnique<FRealTimeStylusInterface>();
+			return TUniquePtr<IStylusInputInterface>(MoveTemp(RealTimeStylusInterface));
+		}
+
+		return nullptr;
+	}
+}
diff --git a/Engine/Plugins/Editor/StylusInput/Source/StylusInputRealTimeStylus/Private/RealTimeStylusInterface.h b/Engine/Plugins/Editor/StylusInput/Source/StylusInputRealTimeStylus/Private/RealTimeStylusInterface.h
new file mode 100644
--- /dev/null
+++ b/Engine/Plugins/Editor/StylusInput/Source/StylusInputRealTimeStylus/Private/RealTimeStylusInterface.h
@@ -0,0 +1,36 @@
+﻿// Copyright Epic Games, Inc. All Rights Reserved.
+
+#pragma once
+
+#include <StylusInput.h>
+#include <StylusInputInterface.h>
+#include <Containers/Map.h>
+#include <Templates/UniquePtr.h>
+
+namespace UE::StylusInput::RealTimeStylus
+{
+	class FRealTimeStylusInstance;
+
+	class FRealTimeStylusInterface : public IStylusInputInterface
+	{
+	public:
+		virtual FName GetName() override;
+
+		virtual IStylusInputInstance* CreateInstance(SWindow& Window) override;
+		virtual bool ReleaseInstance(IStylusInputInstance* Instance) override;
+
+	private:
+		friend class FStylusInputRealTimeStylusModule;
+
+		static TUniquePtr<IStylusInputInterface> Create();
+
+		struct FRefCountedInstance
+		{
+			TUniquePtr<FRealTimeStylusInstance> Instance;
+			int32 RefCount;
+		};
+
+		TMap<SWindow*, FRefCountedInstance> Instances;
+		uint32 NextInstanceID = 0;
+	};
+}
diff --git a/Engine/Plugins/Editor/StylusInput/Source/StylusInputRealTimeStylus/Private/RealTimeStylusPluginAsync.cpp b/Engine/Plugins/Editor/StylusInput/Source/StylusInputRealTimeStylus/Private/RealTimeStylusPluginAsync.cpp
new file mode 100644
--- /dev/null
+++ b/Engine/Plugins/Editor/StylusInput/Source/StylusInputRealTimeStylus/Private/RealTimeStylusPluginAsync.cpp
@@ -0,0 +1,118 @@
+﻿// Copyright Epic Games, Inc. All Rights Reserved.
+
+#include "RealTimeStylusPluginAsync.h"
+
+namespace UE::StylusInput::RealTimeStylus
+{
+	FRealTimeStylusPluginAsync::FRealTimeStylusPluginAsync(IStylusInputInstance* Instance,
+	                                                       FGetWindowContextCallback&& GetWindowContext,
+	                                                       FUpdateTabletContextsCallback&& UpdateTabletContextsCallback,
+	                                                       FUpdateStylusInfoCallback UpdateStylusInfoCallback,
+	                                                       IStylusInputEventHandler* EventHandler)
+		: FRealTimeStylusPluginBase(Instance, MoveTemp(GetWindowContext), MoveTemp(UpdateTabletContextsCallback), MoveTemp(UpdateStylusInfoCallback))
+	{
+		if (EventHandler)
+		{
+			// Immediately install an event handler during construction to capture events coming through during plugin initialization.
+			AddEventHandler(EventHandler);
+		}
+	}
+
+	FRealTimeStylusPluginAsync::~FRealTimeStylusPluginAsync()
+	{
+	}
+
+	HRESULT FRealTimeStylusPluginAsync::RealTimeStylusEnabled(IRealTimeStylus* RealTimeStylus, const ULONG TabletContextIDsCount,
+	                                                              const TABLET_CONTEXT_ID* TabletContextIDs)
+	{
+		return ProcessRealTimeStylusEnabled(RealTimeStylus, TabletContextIDsCount, TabletContextIDs);
+	}
+
+	HRESULT FRealTimeStylusPluginAsync::RealTimeStylusDisabled(IRealTimeStylus* RealTimeStylus, const ULONG TabletContextIDsCount,
+	                                                               const TABLET_CONTEXT_ID* TabletContextIDs)
+	{
+		return ProcessRealTimeStylusDisabled(RealTimeStylus, TabletContextIDsCount, TabletContextIDs);
+	}
+
+	HRESULT FRealTimeStylusPluginAsync::StylusDown(IRealTimeStylus* RealTimeStylus, const StylusInfo* StylusInfo, const ULONG PropertyCount,
+	                                                   LONG* PacketBuffer, LONG**)
+	{
+		return ProcessPackets(StylusInfo, 1, PropertyCount, EPacketType::StylusDown, reinterpret_cast<int32*>(PacketBuffer));
+	}
+
+	HRESULT FRealTimeStylusPluginAsync::StylusUp(IRealTimeStylus* RealTimeStylus, const StylusInfo* StylusInfo, const ULONG PropertyCount,
+	                                                 LONG* PacketBuffer, LONG**)
+	{
+		return ProcessPackets(StylusInfo, 1, PropertyCount, EPacketType::StylusUp, reinterpret_cast<int32*>(PacketBuffer));
+	}
+
+	HRESULT FRealTimeStylusPluginAsync::TabletAdded(IRealTimeStylus* RealTimeStylus, IInkTablet* Tablet)
+	{
+		return ProcessTabletAdded(Tablet);
+	}
+
+	HRESULT FRealTimeStylusPluginAsync::TabletRemoved(IRealTimeStylus* RealTimeStylus, const LONG TabletIndex)
+	{
+		return ProcessTabletRemoved(TabletIndex);
+	}
+
+	HRESULT FRealTimeStylusPluginAsync::InAirPackets(IRealTimeStylus* RealTimeStylus, const StylusInfo* StylusInfo, const ULONG PacketCount,
+	                                                const ULONG PacketBufferLength, LONG* PacketBuffer, ULONG*, LONG**)
+	{
+		return ProcessPackets(StylusInfo, PacketCount, PacketBufferLength, EPacketType::AboveDigitizer, reinterpret_cast<int32*>(PacketBuffer));
+	}
+
+	HRESULT FRealTimeStylusPluginAsync::Packets(IRealTimeStylus* RealTimeStylus, const StylusInfo* StylusInfo, const ULONG PacketCount,
+													const ULONG PacketBufferLength, LONG* PacketBuffer, ULONG*, LONG**)
+	{
+		return ProcessPackets(StylusInfo, PacketCount, PacketBufferLength, EPacketType::OnDigitizer, reinterpret_cast<int32*>(PacketBuffer));
+	}
+
+	HRESULT FRealTimeStylusPluginAsync::StylusInRange(IRealTimeStylus* RealTimeStylus, TABLET_CONTEXT_ID TabletContextID, STYLUS_ID StylusID)
+	{
+		return ProcessStylusInRange(RealTimeStylus, TabletContextID, StylusID);
+	}
+
+	HRESULT FRealTimeStylusPluginAsync::StylusOutOfRange(IRealTimeStylus* RealTimeStylus, TABLET_CONTEXT_ID TabletContextID, STYLUS_ID StylusID)
+	{
+		return E_NOTIMPL;
+	}
+
+	HRESULT FRealTimeStylusPluginAsync::StylusButtonDown(IRealTimeStylus* RealTimeStylus, STYLUS_ID StylusID, const GUID* GuidStylusButton,
+	                                                         POINT* StylusPos)
+	{
+		return E_NOTIMPL;
+	}
+
+	HRESULT FRealTimeStylusPluginAsync::StylusButtonUp(IRealTimeStylus* RealTimeStylus, STYLUS_ID StylusID, const GUID* GuidStylusButton, POINT* StylusPos)
+	{
+		return E_NOTIMPL;
+	}
+
+	HRESULT FRealTimeStylusPluginAsync::CustomStylusDataAdded(IRealTimeStylus* RealTimeStylus, const GUID* GuidId, ULONG DataCount, const BYTE* Data)
+	{
+		return E_NOTIMPL;
+	}
+
+	HRESULT FRealTimeStylusPluginAsync::SystemEvent(IRealTimeStylus* RealTimeStylus, TABLET_CONTEXT_ID TabletContextID, STYLUS_ID StylusID,
+	                                                    SYSTEM_EVENT Event, SYSTEM_EVENT_DATA EventData)
+	{
+		return E_NOTIMPL;
+	}
+
+	HRESULT FRealTimeStylusPluginAsync::Error(IRealTimeStylus*, IStylusPlugin*, const RealTimeStylusDataInterest DataInterest, const HRESULT ErrorCode,
+	                                              LONG_PTR*)
+	{
+		return ProcessError(DataInterest, ErrorCode);
+	}
+
+	HRESULT FRealTimeStylusPluginAsync::UpdateMapping(IRealTimeStylus* RealTimeStylus)
+	{
+		return E_NOTIMPL;
+	}
+
+	HRESULT FRealTimeStylusPluginAsync::DataInterest(RealTimeStylusDataInterest* DataInterest)
+	{
+		return ProcessDataInterest(DataInterest);
+	}
+}
diff --git a/Engine/Plugins/Editor/StylusInput/Source/StylusInputRealTimeStylus/Private/RealTimeStylusPluginAsync.h b/Engine/Plugins/Editor/StylusInput/Source/StylusInputRealTimeStylus/Private/RealTimeStylusPluginAsync.h
new file mode 100644
--- /dev/null
+++ b/Engine/Plugins/Editor/StylusInput/Source/StylusInputRealTimeStylus/Private/RealTimeStylusPluginAsync.h
@@ -0,0 +1,70 @@
+﻿// Copyright Epic Games, Inc. All Rights Reserved.
+
+#pragma once
+
+#include "RealTimeStylusPluginBase.h"
+
+namespace UE::StylusInput::RealTimeStylus
+{
+	class FRealTimeStylusPluginAsync final : public IStylusAsyncPlugin, public FRealTimeStylusPluginBase
+	{
+	public:
+		FRealTimeStylusPluginAsync(IStylusInputInstance* Instance, FGetWindowContextCallback&& GetWindowContext,
+		                           FUpdateTabletContextsCallback&& UpdateTabletContextsCallback, FUpdateStylusInfoCallback UpdateStylusInfoCallback,
+		                           IStylusInputEventHandler* EventHandler);
+		~FRealTimeStylusPluginAsync();
+
+		// IUnknown
+		virtual HRESULT QueryInterface(REFIID InterfaceID, void** InterfaceObject) override
+		{
+			if (InterfaceID == IID_IStylusAsyncPlugin || InterfaceID == IID_IUnknown)
+			{
+				*InterfaceObject = this;
+				AddRef();
+				return S_OK;
+			}
+
+			*InterfaceObject = nullptr;
+			return E_NOINTERFACE;
+		}
+
+		virtual ULONG AddRef() override
+		{
+			return ++RefCount;
+		}
+
+		virtual ULONG Release() override
+		{
+			const int32 NewRefCount = --RefCount;
+			if (NewRefCount == 0)
+				delete this;
+
+			return NewRefCount;
+		}
+
+		// IStylusPlugin
+		virtual HRESULT RealTimeStylusEnabled(IRealTimeStylus* RealTimeStylus, ULONG TabletContextIDsCount, const TABLET_CONTEXT_ID* TabletContextIDs) override;
+		virtual HRESULT RealTimeStylusDisabled(IRealTimeStylus* RealTimeStylus, ULONG TabletContextIDsCount, const TABLET_CONTEXT_ID* TabletContextIDs) override;
+		virtual HRESULT StylusDown(IRealTimeStylus* RealTimeStylus, const StylusInfo* StylusInfo, ULONG PropertyCount, LONG* PacketBuffer, LONG** InOutPkt) override;
+		virtual HRESULT StylusUp(IRealTimeStylus* RealTimeStylus, const StylusInfo* StylusInfo, ULONG PropertyCount, LONG* PacketBuffer, LONG** InOutPkt) override;
+		virtual HRESULT TabletAdded(IRealTimeStylus* RealTimeStylus, IInkTablet* Tablet) override;
+		virtual HRESULT TabletRemoved(IRealTimeStylus* RealTimeStylus, LONG TabletIndex) override;
+		virtual HRESULT InAirPackets(IRealTimeStylus* RealTimeStylus, const StylusInfo* StylusInfo, ULONG PacketCount, ULONG PacketBufferLength, LONG* PacketBuffer, ULONG* InOutPacketCount, LONG** InOutPacketBuffer) override;
+		virtual HRESULT Packets(IRealTimeStylus* RealTimeStylus, const StylusInfo* StylusInfo, ULONG PacketCount, ULONG PacketBufferLength, LONG* PacketBuffer, ULONG* InOutPacketCount, LONG** InOutPacketBuffer) override;
+		virtual HRESULT StylusInRange(IRealTimeStylus* RealTimeStylus, TABLET_CONTEXT_ID TabletContextID, STYLUS_ID StylusID) override;
+		virtual HRESULT StylusOutOfRange(IRealTimeStylus* RealTimeStylus, TABLET_CONTEXT_ID TabletContextID, STYLUS_ID StylusID) override;
+		virtual HRESULT StylusButtonDown(IRealTimeStylus* RealTimeStylus, STYLUS_ID StylusID, const GUID* GuidStylusButton, POINT* StylusPos) override;
+		virtual HRESULT StylusButtonUp(IRealTimeStylus* RealTimeStylus, STYLUS_ID StylusID, const GUID* GuidStylusButton, POINT* StylusPos) override;
+		virtual HRESULT CustomStylusDataAdded(IRealTimeStylus* RealTimeStylus, const GUID* GuidId, ULONG DataCount, const BYTE* Data) override;
+		virtual HRESULT SystemEvent(IRealTimeStylus* RealTimeStylus, TABLET_CONTEXT_ID TabletContextID, STYLUS_ID StylusID, SYSTEM_EVENT Event, SYSTEM_EVENT_DATA EventData) override;
+		virtual HRESULT Error(IRealTimeStylus* RealTimeStylus, IStylusPlugin* Plugin, RealTimeStylusDataInterest DataInterest, HRESULT ErrorCode, LONG_PTR* InternalKey) override;
+		virtual HRESULT UpdateMapping(IRealTimeStylus* RealTimeStylus) override;
+		virtual HRESULT DataInterest(RealTimeStylusDataInterest* DataInterest) override;
+
+	protected:
+		virtual FString GetName() const override { return "OnGameThread"; }
+
+	private:
+		int32 RefCount = 1;
+	};
+}
diff --git a/Engine/Plugins/Editor/StylusInput/Source/StylusInputRealTimeStylus/Private/RealTimeStylusPluginBase.cpp b/Engine/Plugins/Editor/StylusInput/Source/StylusInputRealTimeStylus/Private/RealTimeStylusPluginBase.cpp
new file mode 100644
--- /dev/null
+++ b/Engine/Plugins/Editor/StylusInput/Source/StylusInputRealTimeStylus/Private/RealTimeStylusPluginBase.cpp
@@ -0,0 +1,777 @@
+﻿// Copyright Epic Games, Inc. All Rights Reserved.
+
+#include "RealTimeStylusPluginBase.h"
+
+#include <StylusInput.h>
+#include <StylusInputUtils.h>
+
+#include "RealTimeStylusAPI.h"
+#include "RealTimeStylusUtils.h"
+
+#include <Windows/AllowWindowsPlatformAtomics.h>
+	#include <comdef.h>
+#include <Windows/HideWindowsPlatformAtomics.h>
+
+#define LOG_PREAMBLE "RealTimeStylusPluginBase"
+
+using namespace UE::StylusInput::Private;
+
+namespace UE::StylusInput::RealTimeStylus
+{
+	bool SetupTabletContextMetadata(IInkTablet& InkTablet, FTabletContext& TabletContext)
+	{
+		const FRealTimeStylusAPI& RealTimeStylusAPI = FRealTimeStylusAPI::GetInstance();
+
+		bool bSuccess = true;
+
+		BSTR String = nullptr;
+		if (Succeeded(InkTablet.get_Name(&String), LOG_PREAMBLE))
+		{
+			TabletContext.Name = String;
+			RealTimeStylusAPI.SysFreeString(String);
+		}
+		else
+		{
+			LogError(LOG_PREAMBLE, FString::Format(TEXT("Could not get name for TabletContext with ID {0}."), {TabletContext.ID}));
+			bSuccess = false;
+		}
+
+		if (Succeeded(InkTablet.get_PlugAndPlayId(&String), LOG_PREAMBLE))
+		{
+			TabletContext.PlugAndPlayID = FString(String);
+			RealTimeStylusAPI.SysFreeString(String);
+		}
+		else
+		{
+			LogError(LOG_PREAMBLE, FString::Format(TEXT("Could not get plug and play ID for TabletContext with ID {0}."), {TabletContext.ID}));
+			bSuccess = false;
+		}
+
+		IInkRectangle* Rectangle = nullptr;
+		if (Succeeded(InkTablet.get_MaximumInputRectangle(&Rectangle), LOG_PREAMBLE))
+		{
+			long Top, Left, Bottom, Right;
+			if (Succeeded(Rectangle->GetRectangle(&Top, &Left, &Bottom, &Right), LOG_PREAMBLE))
+			{
+				TabletContext.InputRectangle = {Left, Top, Right, Bottom};
+			}
+		}
+		else
+		{
+			LogError(LOG_PREAMBLE, FString::Format(TEXT("Could not get input rectangle for TabletContext with ID {0}."), {TabletContext.ID}));
+			bSuccess = false;
+		}
+
+		TabletHardwareCapabilities HardwareCapabilities;
+		if (Succeeded(InkTablet.get_HardwareCapabilities(&HardwareCapabilities), LOG_PREAMBLE))
+		{
+			TabletContext.HardwareCapabilities =
+				static_cast<ETabletHardwareCapabilities>(HardwareCapabilities & THWC_Integrated) |
+				static_cast<ETabletHardwareCapabilities>(HardwareCapabilities & THWC_CursorMustTouch) |
+				static_cast<ETabletHardwareCapabilities>(HardwareCapabilities & THWC_HardProximity) |
+				static_cast<ETabletHardwareCapabilities>(HardwareCapabilities & THWC_CursorsHavePhysicalIds);
+		}
+		else
+		{
+			LogError(LOG_PREAMBLE, FString::Format(TEXT("Could not get hardware capabilities for TabletContext with ID {0}."), {TabletContext.ID}));
+			bSuccess = false;
+		}
+
+		return bSuccess;
+	}
+
+	// TODO compact number of functions needed to only have differences in behavior, and then use offsetof() to apply to different members
+	// TODO XY validate mapping 
+
+	void SetPropertyX(FStylusInputPacket& Packet, const int32 Value, const int8* Data)
+	{
+		const double& Scale = reinterpret_cast<const double*>(Data)[0];
+		const double& Maximum = reinterpret_cast<const double*>(Data)[1];
+		const double& WindowWidth = reinterpret_cast<const double*>(Data)[2];
+
+		Packet.X = Value / Scale / Maximum * WindowWidth;
+	}
+	static_assert(std::is_same_v<FPacketPropertyHandler::FFuncSetProperty, decltype(SetPropertyX)>);
+
+	void SetPropertyY(FStylusInputPacket& Packet, const int32 Value, const int8* Data)
+	{
+		const double& Scale = reinterpret_cast<const double*>(Data)[0];
+		const double& Maximum = reinterpret_cast<const double*>(Data)[1];
+		const double& WindowHeight = reinterpret_cast<const double*>(Data)[2];
+
+		Packet.Y = Value / Scale / Maximum * WindowHeight;
+	}
+	static_assert(std::is_same_v<FPacketPropertyHandler::FFuncSetProperty, decltype(SetPropertyY)>);
+
+	void SetPropertyZ(FStylusInputPacket& Packet, const int32 Value, const int8* Data)
+	{
+		const int32& Minimum = reinterpret_cast<const int32*>(Data)[0];
+		const int32& Maximum = reinterpret_cast<const int32*>(Data)[1];
+
+		Packet.Z = static_cast<float>(Value - Minimum) / (Maximum - Minimum);
+	}
+	static_assert(std::is_same_v<FPacketPropertyHandler::FFuncSetProperty, decltype(SetPropertyZ)>);
+
+	void SetPropertyPacketStatus(FStylusInputPacket& Packet, const int32 Value, const int8*)
+	{
+		Packet.PenStatus = static_cast<EPenStatus>(Value);
+	}
+	static_assert(std::is_same_v<FPacketPropertyHandler::FFuncSetProperty, decltype(SetPropertyPacketStatus)>);
+
+	void SetPropertyTimerTick(FStylusInputPacket& Packet, const int32 Value, const int8* Data)
+	{
+#if 0
+		Packet.TimerTick = Value; // The value provided from the RealTimeStylus API seems to only be non-zero if the stylus is down.
+#else
+		FRealTimeStylusAPI::FFuncGetTickCount* const FuncGetTickCount =	*reinterpret_cast<decltype(FuncGetTickCount)*>(Data);
+		Packet.TimerTick = FuncGetTickCount();
+#endif
+	}
+	static_assert(std::is_same_v<FPacketPropertyHandler::FFuncSetProperty, decltype(SetPropertyTimerTick)>);
+
+	void SetPropertySerialNumber(FStylusInputPacket& Packet, const int32 Value, const int8* Data)
+	{
+		Packet.SerialNumber = Value;
+	}
+	static_assert(std::is_same_v<FPacketPropertyHandler::FFuncSetProperty, decltype(SetPropertySerialNumber)>);
+
+	void SetPropertyNormalPressure(FStylusInputPacket& Packet, const int32 Value, const int8* Data)
+	{
+		const int32& Minimum = reinterpret_cast<const int32*>(Data)[0];
+		const int32& Maximum = reinterpret_cast<const int32*>(Data)[1];
+
+		Packet.NormalPressure = static_cast<float>(Value - Minimum) / (Maximum - Minimum);
+	}
+	static_assert(std::is_same_v<FPacketPropertyHandler::FFuncSetProperty, decltype(SetPropertyNormalPressure)>);
+
+	void SetPropertyTangentPressure(FStylusInputPacket& Packet, const int32 Value, const int8* Data)
+	{
+		const int32& Minimum = reinterpret_cast<const int32*>(Data)[0];
+		const int32& Maximum = reinterpret_cast<const int32*>(Data)[1];
+
+		Packet.TangentPressure = static_cast<float>(Value - Minimum) / (Maximum - Minimum);
+	}
+	static_assert(std::is_same_v<FPacketPropertyHandler::FFuncSetProperty, decltype(SetPropertyTangentPressure)>);
+
+	void SetPropertyButtonPressure(FStylusInputPacket& Packet, const int32 Value, const int8* Data)
+	{
+		const int32& Minimum = reinterpret_cast<const int32*>(Data)[0];
+		const int32& Maximum = reinterpret_cast<const int32*>(Data)[1];
+
+		Packet.ButtonPressure = static_cast<float>(Value - Minimum) / (Maximum - Minimum);
+	}
+	static_assert(std::is_same_v<FPacketPropertyHandler::FFuncSetProperty, decltype(SetPropertyButtonPressure)>);
+
+	void SetPropertyXTiltOrientation(FStylusInputPacket& Packet, const int32 Value, const int8* Data)
+	{
+		const float& Resolution = reinterpret_cast<const float*>(Data)[0];
+
+		Packet.XTiltOrientation = Value / Resolution;
+	}
+	static_assert(std::is_same_v<FPacketPropertyHandler::FFuncSetProperty, decltype(SetPropertyXTiltOrientation)>);
+
+	void SetPropertyYTiltOrientation(FStylusInputPacket& Packet, const int32 Value, const int8* Data)
+	{
+		const float& Resolution = reinterpret_cast<const float*>(Data)[0];
+
+		Packet.YTiltOrientation = Value / Resolution;
+	}
+	static_assert(std::is_same_v<FPacketPropertyHandler::FFuncSetProperty, decltype(SetPropertyYTiltOrientation)>);
+
+	void SetPropertyAzimuthOrientation(FStylusInputPacket& Packet, const int32 Value, const int8* Data)
+	{
+		const float& Resolution = reinterpret_cast<const float*>(Data)[0];
+
+		Packet.AzimuthOrientation = Value / Resolution;
+	}
+	static_assert(std::is_same_v<FPacketPropertyHandler::FFuncSetProperty, decltype(SetPropertyAzimuthOrientation)>);
+
+	void SetPropertyAltitudeOrientation(FStylusInputPacket& Packet, const int32 Value, const int8* Data)
+	{
+		const float& Resolution = reinterpret_cast<const float*>(Data)[0];
+
+		Packet.AltitudeOrientation = Value / Resolution;
+	}
+	static_assert(std::is_same_v<FPacketPropertyHandler::FFuncSetProperty, decltype(SetPropertyAltitudeOrientation)>);
+
+	void SetPropertyTwistOrientation(FStylusInputPacket& Packet, const int32 Value, const int8* Data)
+	{
+		const float& Resolution = reinterpret_cast<const float*>(Data)[0];
+
+		Packet.TwistOrientation = Value / Resolution;
+	}
+	static_assert(std::is_same_v<FPacketPropertyHandler::FFuncSetProperty, decltype(SetPropertyTwistOrientation)>);
+
+	void SetPropertyPitchRotation(FStylusInputPacket& Packet, const int32 Value, const int8* Data)
+	{
+		const float& Resolution = reinterpret_cast<const float*>(Data)[0];
+
+		Packet.PitchRotation = Value / Resolution;
+	}
+	static_assert(std::is_same_v<FPacketPropertyHandler::FFuncSetProperty, decltype(SetPropertyPitchRotation)>);
+
+	void SetPropertyRollRotation(FStylusInputPacket& Packet, const int32 Value, const int8* Data)
+	{
+		const float& Resolution = reinterpret_cast<const float*>(Data)[0];
+
+		Packet.RollRotation = Value / Resolution;
+	}
+	static_assert(std::is_same_v<FPacketPropertyHandler::FFuncSetProperty, decltype(SetPropertyRollRotation)>);
+
+	void SetPropertyYawRotation(FStylusInputPacket& Packet, const int32 Value, const int8* Data)
+	{
+		const float& Resolution = reinterpret_cast<const float*>(Data)[0];
+
+		Packet.YawRotation = Value / Resolution;
+	}
+	static_assert(std::is_same_v<FPacketPropertyHandler::FFuncSetProperty, decltype(SetPropertyYawRotation)>);
+
+	void SetPropertyWidth(FStylusInputPacket& Packet, const int32 Value, const int8* Data)
+	{
+		const float& Resolution = reinterpret_cast<const float*>(Data)[0];
+
+		Packet.Width = Value / Resolution;
+	}
+	static_assert(std::is_same_v<FPacketPropertyHandler::FFuncSetProperty, decltype(SetPropertyWidth)>);
+
+	void SetPropertyHeight(FStylusInputPacket& Packet, const int32 Value, const int8* Data)
+	{
+		const float& Resolution = reinterpret_cast<const float*>(Data)[0];
+
+		Packet.Height = Value / Resolution;
+	}
+	static_assert(std::is_same_v<FPacketPropertyHandler::FFuncSetProperty, decltype(SetPropertyHeight)>);
+
+	void SetPropertyFingerContactConfidence(FStylusInputPacket& Packet, const int32 Value, const int8* Data)
+	{
+		const int32& Minimum = reinterpret_cast<const int32*>(Data)[0];
+		const int32& Maximum = reinterpret_cast<const int32*>(Data)[1];
+
+		Packet.FingerContactConfidence = static_cast<float>(Value - Minimum) / (Maximum - Minimum);
+	}
+	static_assert(std::is_same_v<FPacketPropertyHandler::FFuncSetProperty, decltype(SetPropertyFingerContactConfidence)>);
+
+	void SetPropertyDeviceContactID(FStylusInputPacket& Packet, const int32 Value, const int8* Data)
+	{
+		Packet.DeviceContactID = Value;
+	}
+	static_assert(std::is_same_v<FPacketPropertyHandler::FFuncSetProperty, decltype(SetPropertyDeviceContactID)>);
+
+	FPacketPropertyHandler::FFuncSetProperty* GetSetPropertyFunc(const EPacketPropertyType Type)
+	{
+		switch (Type)
+		{
+		case EPacketPropertyType::X:
+			return &SetPropertyX;
+		case EPacketPropertyType::Y:
+			return &SetPropertyY;
+		case EPacketPropertyType::Z:
+			return &SetPropertyZ;
+		case EPacketPropertyType::PacketStatus:
+			return &SetPropertyPacketStatus;
+		case EPacketPropertyType::TimerTick:
+			return &SetPropertyTimerTick;
+		case EPacketPropertyType::SerialNumber:
+			return &SetPropertySerialNumber;
+		case EPacketPropertyType::NormalPressure:
+			return &SetPropertyNormalPressure;
+		case EPacketPropertyType::TangentPressure:
+			return &SetPropertyTangentPressure;
+		case EPacketPropertyType::ButtonPressure:
+			return &SetPropertyButtonPressure;
+		case EPacketPropertyType::XTiltOrientation:
+			return &SetPropertyXTiltOrientation;
+		case EPacketPropertyType::YTiltOrientation:
+			return &SetPropertyYTiltOrientation;
+		case EPacketPropertyType::AzimuthOrientation:
+			return &SetPropertyAzimuthOrientation;
+		case EPacketPropertyType::AltitudeOrientation:
+			return &SetPropertyAltitudeOrientation;
+		case EPacketPropertyType::TwistOrientation:
+			return &SetPropertyTwistOrientation;
+		case EPacketPropertyType::PitchRotation:
+			return &SetPropertyPitchRotation;
+		case EPacketPropertyType::RollRotation:
+			return &SetPropertyRollRotation;
+		case EPacketPropertyType::YawRotation:
+			return &SetPropertyYawRotation;
+		case EPacketPropertyType::Width:
+			return &SetPropertyWidth;
+		case EPacketPropertyType::Height:
+			return &SetPropertyHeight;
+		case EPacketPropertyType::FingerContactConfidence:
+			return &SetPropertyFingerContactConfidence;
+		case EPacketPropertyType::DeviceContactID:
+			return &SetPropertyDeviceContactID;
+		default:
+			checkSlow(false)
+			return nullptr;
+		}
+	}
+
+	template <int32 DataBufferLength>
+	void AssignSetPropertyData(const EPacketPropertyType Type, int8 *const Data, const FWindowContext& WindowContext, const FPacketProperty& Description)
+	{
+		int8* DataCurrent = Data;
+
+		auto Add = [&DataCurrent]<typename T>(T Value)
+		{
+			*reinterpret_cast<T*>(DataCurrent) = Value;
+			DataCurrent += sizeof(T);
+		};
+
+		switch (Type)
+		{
+		case EPacketPropertyType::X:
+			Add(WindowContext.XYScale.X);
+			Add(WindowContext.XYMaximum.X);
+			Add(WindowContext.WindowSize.X);
+			break;
+		case EPacketPropertyType::Y:
+			Add(WindowContext.XYScale.Y);
+			Add(WindowContext.XYMaximum.Y);
+			Add(WindowContext.WindowSize.Y);
+			break;
+		case EPacketPropertyType::Z:
+		case EPacketPropertyType::NormalPressure:
+		case EPacketPropertyType::TangentPressure:
+		case EPacketPropertyType::ButtonPressure:
+		case EPacketPropertyType::FingerContactConfidence:
+			Add(Description.Minimum);
+			Add(Description.Maximum);
+			break;
+		case EPacketPropertyType::XTiltOrientation:
+		case EPacketPropertyType::YTiltOrientation:
+		case EPacketPropertyType::AzimuthOrientation:
+		case EPacketPropertyType::AltitudeOrientation:
+		case EPacketPropertyType::TwistOrientation:
+		case EPacketPropertyType::PitchRotation:
+		case EPacketPropertyType::RollRotation:
+		case EPacketPropertyType::YawRotation:
+		case EPacketPropertyType::Width:
+		case EPacketPropertyType::Height:
+			Add(Description.Resolution);
+			break;
+		case EPacketPropertyType::TimerTick:
+			Add(FRealTimeStylusAPI::GetInstance().GetTickCount);
+			break;
+		case EPacketPropertyType::PacketStatus:
+		case EPacketPropertyType::SerialNumber:
+		case EPacketPropertyType::DeviceContactID:
+			// do nothing
+			break;
+		default:
+			checkSlow(false)
+		}
+
+		check(std::distance(Data, DataCurrent) <= DataBufferLength)
+	}
+
+	bool SetupTabletContextPacketDescriptionData(IRealTimeStylus* RealTimeStylus, const FWindowContext& WindowContext, FTabletContext& TabletContext)
+	{
+		const FRealTimeStylusAPI& RealTimeStylusAPI = FRealTimeStylusAPI::GetInstance();
+
+		auto InvalidatePacketDescriptions = [&PacketDescriptions = TabletContext.PacketProperties](const int32 StartingIndex)
+		{
+			for (int32 Index = StartingIndex; Index < static_cast<int32>(EPacketPropertyType::Num_Enumerators); ++Index)
+			{
+				PacketDescriptions[Index] = {};
+			}
+		};
+
+		FLOAT InkToDeviceScaleX = 0.0f;
+		FLOAT InkToDeviceScaleY = 0.0f;
+		ULONG PropertiesCount = 0;
+		PACKET_PROPERTY* Properties = nullptr;
+		if (Succeeded(RealTimeStylus->GetPacketDescriptionData(TabletContext.ID, &InkToDeviceScaleX, &InkToDeviceScaleY, &PropertiesCount, &Properties), LOG_PREAMBLE))
+		{
+			if (!Properties)
+			{
+				LogWarning(LOG_PREAMBLE, FString::Format(
+					           TEXT("Retrieved nullptr when trying to get packet description data for TabletContext with ID {0}."), {TabletContext.ID}));
+				return false;
+			}
+
+			if (PropertiesCount == 0)
+			{
+				LogWarning(LOG_PREAMBLE, FString::Format(
+					           TEXT("Retrieved zero packets properties when trying to get packet description data for TabletContext with ID {0}."),
+					           {TabletContext.ID}));
+				return false;
+			}
+
+			for (ULONG Index = 0; Index < PropertiesCount; ++Index)
+			{
+#pragma warning(push)
+#pragma warning(disable : 28199)
+				const PACKET_PROPERTY& Property = Properties[Index];
+#pragma warning(pop)
+
+				bool bFoundProperty = false;
+				for (const FPacketPropertyConstant& PropertyConstant : PacketPropertyConstants)
+				{
+					if (Property.guid == PropertyConstant.Guid)
+					{
+						const ETabletSupportedProperties SupportedProperty = static_cast<ETabletSupportedProperties>(
+							1 << static_cast<std::underlying_type_t<ETabletSupportedProperties>>(PropertyConstant.PacketPropertyType)
+						);
+						TabletContext.SupportedProperties = TabletContext.SupportedProperties | SupportedProperty;
+
+						FPacketProperty& PropertyProperty = TabletContext.PacketProperties[Index];
+						PropertyProperty.Type = static_cast<EPacketPropertyType>(PropertyConstant.PacketPropertyType);
+						PropertyProperty.Minimum = Property.PropertyMetrics.nLogicalMin;
+						PropertyProperty.Maximum = Property.PropertyMetrics.nLogicalMax;
+						PropertyProperty.MetricUnit = static_cast<ETabletPropertyMetricUnit>(Property.PropertyMetrics.Units);
+						PropertyProperty.Resolution = Property.PropertyMetrics.fResolution;
+
+						FPacketPropertyHandler& PropertyPropertyHandler = TabletContext.PacketPropertyHandlers[Index];
+						PropertyPropertyHandler.SetProperty = GetSetPropertyFunc(PropertyConstant.PacketPropertyType);
+						AssignSetPropertyData<FPacketPropertyHandler::SetPropertyDataBufferLength>(
+							PropertyConstant.PacketPropertyType, PropertyPropertyHandler.SetPropertyData, WindowContext, PropertyProperty);
+
+						if (static_cast<int8>(PropertyProperty.MetricUnit) >= static_cast<int8>(ETabletPropertyMetricUnit::Num_Enumerators))
+						{
+							PropertyProperty.MetricUnit = ETabletPropertyMetricUnit::Default;
+
+							LogWarning(LOG_PREAMBLE, FString::Format(
+									   TEXT("Encountered unknown metric unit value '{0}' while evaluating packet description data for TabletContext with ID {1}."),
+									   {static_cast<int8>(PropertyProperty.MetricUnit), TabletContext.ID}));
+						}
+
+						bFoundProperty = true;
+						break;
+					}
+				}
+
+				if (!bFoundProperty)
+				{
+					wchar_t GUIDStringBuffer[64];
+					const int32 GUIDStringLength = RealTimeStylusAPI.StringFromGUID2(Property.guid, GUIDStringBuffer, std::size(GUIDStringBuffer));
+					if (ensure(GUIDStringLength > 0))
+					{
+						LogWarning(LOG_PREAMBLE, FString::Format(
+							           TEXT("Encountered unknown property '{0}' while evaluating packet description data for TabletContext with ID {1}."),
+							           {GUIDStringBuffer, TabletContext.ID}));
+					}
+				}
+			}
+
+			RealTimeStylusAPI.CoTaskMemFree(Properties);
+
+			InvalidatePacketDescriptions(PropertiesCount);
+		}
+		else
+		{
+			LogError(LOG_PREAMBLE, FString::Format(TEXT("Could not get packet description data for TabletContext with ID {0}."), {TabletContext.ID}));
+
+			InvalidatePacketDescriptions(0);
+
+			return false;
+		}
+
+		return true;
+	}
+
+	bool SetupTabletContext(IRealTimeStylus* RealTimeStylus, const FWindowContext& WindowContext, FTabletContext& TabletContext)
+	{
+		IInkTablet* InkTablet = nullptr;
+
+		if (Failed(RealTimeStylus->GetTabletFromTabletContextId(TabletContext.ID, &InkTablet), LOG_PREAMBLE))
+		{
+			LogError(LOG_PREAMBLE, FString::Format(TEXT("Could not get tablet context data for ID {0}."), {TabletContext.ID}));
+			return false;
+		}
+
+		bool bSuccess = true;
+
+		bSuccess &= SetupTabletContextMetadata(*InkTablet, TabletContext);
+		bSuccess &= SetupTabletContextPacketDescriptionData(RealTimeStylus, WindowContext, TabletContext);
+
+		return bSuccess;
+	}
+
+	bool FRealTimeStylusPluginBase::AddEventHandler(IStylusInputEventHandler* EventHandler)
+	{
+		check(EventHandler)
+
+		if (EventHandlers.Contains(EventHandler))
+		{
+			LogWarning(LOG_PREAMBLE, FString::Format(TEXT("Event handler '{0}' already exists in {1} plugin."), {EventHandler->GetName(), GetName()}));
+			return false;
+		}
+
+		EventHandlers.Add(EventHandler);
+
+		LogVerbose(LOG_PREAMBLE, FString::Format(TEXT("Event handler '{0}' was added to {1} plugin."), {EventHandler->GetName(), GetName()}));
+
+		return true;
+	}
+
+	bool FRealTimeStylusPluginBase::RemoveEventHandler(IStylusInputEventHandler* EventHandler)
+	{
+		check(EventHandler)
+
+		const bool bWasRemoved = EventHandlers.Remove(EventHandler) > 0;
+
+		if (bWasRemoved)
+		{
+			LogVerbose(LOG_PREAMBLE, FString::Format(TEXT("Event handler '{0}' was removed from {1} plugin."), {EventHandler->GetName(), GetName()}));
+		}
+
+		return bWasRemoved;
+	}
+
+	FRealTimeStylusPluginBase::FRealTimeStylusPluginBase(IStylusInputInstance* Instance,
+	                                                     FGetWindowContextCallback&& GetWindowContextCallback,
+	                                                     FUpdateTabletContextsCallback&& UpdateTabletContextsCallback,
+	                                                     FUpdateStylusInfoCallback&& UpdateStylusInfoCallback)
+		: Instance(Instance)
+		, GetWindowContextCallback(MoveTemp(GetWindowContextCallback))
+		, UpdateTabletContextsCallback(MoveTemp(UpdateTabletContextsCallback))
+		, UpdateStylusInfoCallback(MoveTemp(UpdateStylusInfoCallback))
+	{
+		check(this->GetWindowContextCallback.IsBound())
+		check(this->UpdateTabletContextsCallback.IsBound())
+	}
+
+	void FRealTimeStylusPluginBase::DebugEvent(const FString& Message) const
+	{
+		for (IStylusInputEventHandler* EventHandler : EventHandlers)
+		{
+			EventHandler->OnDebugEvent(Message, Instance);
+		}
+	}
+
+	HRESULT FRealTimeStylusPluginBase::ProcessDataInterest(RealTimeStylusDataInterest* DataInterest)
+	{
+		constexpr bool bGetAllData = false;
+
+		if (bGetAllData)
+		{
+			*DataInterest = RTSDI_AllData;
+
+			DebugEvent("Requested all data from stylus input.");
+		}
+		else
+		{
+			*DataInterest = static_cast<RealTimeStylusDataInterest>(
+				RTSDI_Error |
+				RTSDI_RealTimeStylusEnabled | RTSDI_RealTimeStylusDisabled |
+				RTSDI_InAirPackets | RTSDI_Packets |
+				RTSDI_StylusInRange | RTSDI_StylusDown | RTSDI_StylusUp);
+
+			DebugEvent("Requested data for Error, RealTimeStylusEnabled, RealTimeStylusDisabled, InAirPackets, Packets, StylusInRange, StylusDown, "
+				"StylusUp from stylus input (DataInterest event).");
+		}
+
+		return S_OK;
+	}
+
+	HRESULT FRealTimeStylusPluginBase::ProcessError(RealTimeStylusDataInterest DataInterest, const HRESULT ErrorCode)
+	{
+		constexpr bool bShowAllErrors = false;
+
+		if (bShowAllErrors || ErrorCode != E_NOTIMPL)
+		{
+			const char* DataInterestStr = [DataInterest]
+			{
+				switch (DataInterest)
+				{
+				case RTSDI_AllData:
+					return "AllData";
+				case RTSDI_None:
+					return "None";
+				case RTSDI_Error:
+					return "Error";
+				case RTSDI_RealTimeStylusEnabled:
+					return "RealTimeStylusEnabled";
+				case RTSDI_RealTimeStylusDisabled:
+					return "RealTimeStylusDisabled";
+				case RTSDI_StylusNew:
+					return "StylusNew";
+				case RTSDI_StylusInRange:
+					return "StylusInRange";
+				case RTSDI_InAirPackets:
+					return "InAirPackets";
+				case RTSDI_StylusOutOfRange:
+					return "StylusOutOfRange";
+				case RTSDI_StylusDown:
+					return "StylusDown";
+				case RTSDI_Packets:
+					return "Packets";
+				case RTSDI_StylusUp:
+					return "StylusUp";
+				case RTSDI_StylusButtonUp:
+					return "StylusButtonUp";
+				case RTSDI_StylusButtonDown:
+					return "StylusButtonDown";
+				case RTSDI_SystemEvents:
+					return "SystemEvents";
+				case RTSDI_TabletAdded:
+					return "TabletAdded";
+				case RTSDI_TabletRemoved:
+					return "TabletRemoved";
+				case RTSDI_CustomStylusDataAdded:
+					return "CustomStylusDataAdded";
+				case RTSDI_UpdateMapping:
+					return "UpdateMapping";
+				case RTSDI_DefaultEvents:
+					return "DefaultEvents";
+				}
+				return "<unknown>";
+			}();
+
+			const FString ErrorMessage = FString::Format(TEXT("Error in {0} plugin: {1}, Error={2} ({3})."), {
+				                                             GetName(), DataInterestStr, _com_error(ErrorCode).ErrorMessage(), static_cast<int32>(ErrorCode)
+			                                             });
+
+			DebugEvent(ErrorMessage);
+		}
+
+		return S_OK;
+	}
+
+	HRESULT FRealTimeStylusPluginBase::ProcessPackets(const StylusInfo* StylusInfo, uint32 PacketCount, uint32 PacketBufferLength, EPacketType Type,
+	                                                  const int32* PacketBuffer)
+	{
+		checkSlow(StylusInfo)
+		checkSlow(PacketBuffer)
+
+		const FPacketPropertyHandler* Handlers = GetPacketPropertyHandlers(StylusInfo->tcid);
+		if (!Handlers)
+		{
+			return E_FAIL;
+		}
+
+		const uint32 PropertyCount = PacketBufferLength / PacketCount;
+		checkSlow(PropertyCount < static_cast<uint32>(EPacketPropertyType::Num_Enumerators));
+
+		for (uint32 PacketIndex = 0; PacketIndex < PacketCount; PacketIndex += PropertyCount)
+		{
+			PacketStats.NewPacket();
+
+			FStylusInputPacket Packet{StylusInfo->tcid, StylusInfo->cid, Type};
+
+			for (uint32 PropertyIndex = 0; PropertyIndex < PropertyCount; ++PropertyIndex)
+			{
+				const FPacketPropertyHandler& Handler = Handlers[PropertyIndex];
+				checkSlow(Handler.SetProperty)
+				Handler.SetProperty(Packet, PacketBuffer[PropertyIndex], Handler.SetPropertyData);
+			}
+
+			for (IStylusInputEventHandler* EventHandler : EventHandlers)
+			{
+				EventHandler->OnPacket(Packet, Instance);
+			}
+
+			PacketBuffer += PropertyCount;
+		}
+
+		return S_OK;
+	}
+
+	HRESULT FRealTimeStylusPluginBase::ProcessRealTimeStylusEnabled(IRealTimeStylus* RealTimeStylus, const uint32 TabletContextIDsCount,
+																   const TABLET_CONTEXT_ID* TabletContextIDs)
+	{
+		DebugEvent("Stylus input was enabled (RealTimeStylusEnabled event).");
+
+		return UpdateTabletContexts(RealTimeStylus, TabletContextIDsCount, TabletContextIDs);
+	}
+
+	HRESULT FRealTimeStylusPluginBase::ProcessRealTimeStylusDisabled(IRealTimeStylus* RealTimeStylus, const uint32 TabletContextIDsCount,
+																	const TABLET_CONTEXT_ID* TabletContextIDs)
+	{
+		DebugEvent("Stylus input was disabled (RealTimeStylusDisabled event).");
+
+		return UpdateTabletContexts(RealTimeStylus, TabletContextIDsCount, TabletContextIDs);
+	}
+
+	HRESULT FRealTimeStylusPluginBase::ProcessStylusInRange(IRealTimeStylus* RealTimeStylus, TABLET_CONTEXT_ID TabletContextID, STYLUS_ID StylusID)
+	{
+		UpdateStylusInfoCallback.Execute(StylusID);
+		return S_OK;
+	}
+
+	HRESULT FRealTimeStylusPluginBase::ProcessTabletAdded(IInkTablet* Tablet)
+	{
+		// After a TabletAdded event, a RealTimeStylusDisabled event will fire directly followed by a RealTimeStylusEnabled event.
+		// We are using these two events instead to update the tablet contexts.
+
+		return E_NOTIMPL;
+	}
+
+	HRESULT FRealTimeStylusPluginBase::ProcessTabletRemoved(LONG TabletIndex)
+	{
+		// For simplicity, we don't remove a tablet context when this event is received.
+		// However, since there are no more packets coming through for the removed tablet there should be nothing that's continuing to access the outdated
+		// tablet context data.
+
+		return E_NOTIMPL;
+	}
+
+	const FPacketPropertyHandler* FRealTimeStylusPluginBase::GetPacketPropertyHandlers(uint32 TabletContextID) const
+	{
+		const TSharedPtr<FTabletContext> TabletContext = TabletContexts.Get(TabletContextID);
+		return TabletContext ? TabletContext->PacketPropertyHandlers : nullptr;
+	}
+
+	HRESULT FRealTimeStylusPluginBase::UpdateTabletContexts(IRealTimeStylus* RealTimeStylus, uint32 TabletContextIDsNum, const TABLET_CONTEXT_ID* TabletContextIDs)
+	{
+		const FWindowContext& WindowContext = GetWindowContextCallback.Execute();
+
+		bool bSuccess = true;
+
+		// Remove outdated tablet contexts
+		uint32 TabletContextIndex = 0;
+		while (TabletContextIndex < TabletContexts.Num())
+		{
+			const uint32 ExistingTabletContextID = TabletContexts[TabletContextIndex]->ID;
+
+			const bool bContainsID = [ExistingTabletContextID, TabletContextIDsNum, TabletContextIDs]
+			{
+				for (uint32 TabletContextIDsIndex = 0; TabletContextIDsIndex < TabletContextIDsNum; ++TabletContextIDsIndex)
+				{
+					if (ExistingTabletContextID == TabletContextIDs[TabletContextIDsIndex])
+					{
+						return true;
+					}
+				}
+				return false;
+			}();
+
+			if (bContainsID)
+			{
+				++TabletContextIndex;
+			}
+			else
+			{
+				TabletContexts.Remove(ExistingTabletContextID);
+
+				LogVerbose(LOG_PREAMBLE, FString::Format(TEXT("Removed tablet context data for ID {0}."), {ExistingTabletContextID}));
+			}
+		}
+
+		// Add new tablet contexts
+		for (uint32 TabletContextIDsIndex = 0; TabletContextIDsIndex < TabletContextIDsNum; ++TabletContextIDsIndex)
+		{
+			const uint32 TabletContextID = TabletContextIDs[TabletContextIDsIndex];
+			if (!TabletContexts.Contains(TabletContextID))
+			{
+				const TSharedRef<FTabletContext> TabletContextRef = TabletContexts.Add(TabletContextID);
+				FTabletContext& TabletContext = TabletContextRef.Get();
+
+				bSuccess &= SetupTabletContext(RealTimeStylus, WindowContext, TabletContext);
+
+				LogVerbose(LOG_PREAMBLE, FString::Format(TEXT("Added TabletContext for ID {0} [{1}, {2}]."), {
+															TabletContext.ID, TabletContext.Name, TabletContext.PlugAndPlayID
+														}));
+			}
+		}
+
+		UpdateTabletContextsCallback.Execute(TabletContexts);
+
+		return bSuccess ? S_OK : E_FAIL;
+	}
+}
+
+#undef LOG_PREAMBLE
diff --git a/Engine/Plugins/Editor/StylusInput/Source/StylusInputRealTimeStylus/Private/RealTimeStylusPluginBase.h b/Engine/Plugins/Editor/StylusInput/Source/StylusInputRealTimeStylus/Private/RealTimeStylusPluginBase.h
new file mode 100644
--- /dev/null
+++ b/Engine/Plugins/Editor/StylusInput/Source/StylusInputRealTimeStylus/Private/RealTimeStylusPluginBase.h
@@ -0,0 +1,69 @@
+﻿// Copyright Epic Games, Inc. All Rights Reserved.
+
+#pragma once
+
+#include <StylusInputPacket.h>
+#include <Delegates/Delegate.h>
+
+#include "RealTimeStylusStats.h"
+#include "RealTimeStylusTabletContext.h"
+
+#include <Windows/AllowWindowsPlatformTypes.h>
+	#include <Microsoft/COMPointer.h>
+	#include <RTSCom.h>
+#include <Windows/HideWindowsPlatformTypes.h>
+
+namespace UE::StylusInput
+{
+	class IStylusInputEventHandler;
+	class IStylusInputInstance;
+}
+
+namespace UE::StylusInput::RealTimeStylus
+{
+	class FRealTimeStylusAPI;
+
+	DECLARE_DELEGATE_RetVal(const FWindowContext&, FGetWindowContextCallback);
+	DECLARE_DELEGATE_OneParam(FUpdateTabletContextsCallback, const FTabletContextContainer&);
+	DECLARE_DELEGATE_OneParam(FUpdateStylusInfoCallback, STYLUS_ID);
+
+	class FRealTimeStylusPluginBase
+	{
+	public:
+		bool AddEventHandler(IStylusInputEventHandler* EventHandler);
+		bool RemoveEventHandler(IStylusInputEventHandler* EventHandler);
+		int32 NumEventHandlers() const { return EventHandlers.Num(); }
+
+		float GetPacketsPerSecond() const { return PacketStats.GetPacketsPerSecond(); }
+
+	protected:
+		FRealTimeStylusPluginBase(IStylusInputInstance* Instance, FGetWindowContextCallback&& GetWindowContextCallback,
+		                          FUpdateTabletContextsCallback&& UpdateTabletContextsCallback, FUpdateStylusInfoCallback&& UpdateStylusInfoCallback);
+		~FRealTimeStylusPluginBase() = default;
+
+		virtual FString GetName() const = 0;
+
+		void DebugEvent(const FString& Message) const;
+
+		HRESULT ProcessDataInterest(RealTimeStylusDataInterest* DataInterest);
+		HRESULT ProcessError(RealTimeStylusDataInterest DataInterest, HRESULT ErrorCode);
+		HRESULT ProcessPackets(const StylusInfo* StylusInfo, uint32 PacketCount, uint32 PacketBufferLength, EPacketType Type, const int32* PacketBuffer);
+		HRESULT ProcessRealTimeStylusEnabled(IRealTimeStylus* RealTimeStylus, uint32 TabletContextIDsCount, const TABLET_CONTEXT_ID* TabletContextIDs);
+		HRESULT ProcessRealTimeStylusDisabled(IRealTimeStylus* RealTimeStylus, uint32 TabletContextIDsCount, const TABLET_CONTEXT_ID* TabletContextIDs);
+		HRESULT ProcessStylusInRange(IRealTimeStylus* RealTimeStylus, TABLET_CONTEXT_ID TabletContextID, STYLUS_ID StylusID);
+		HRESULT ProcessTabletAdded(IInkTablet* Tablet);
+		HRESULT ProcessTabletRemoved(LONG TabletIndex);
+
+	private:
+		const FPacketPropertyHandler* GetPacketPropertyHandlers(uint32 TabletContextID) const;
+		HRESULT UpdateTabletContexts(IRealTimeStylus* RealTimeStylus, uint32 TabletContextIDsNum, const TABLET_CONTEXT_ID* TabletContextIDs);
+
+		IStylusInputInstance *const Instance;
+		FPacketStats PacketStats;
+		FGetWindowContextCallback GetWindowContextCallback;
+		FUpdateTabletContextsCallback UpdateTabletContextsCallback;
+		FUpdateStylusInfoCallback UpdateStylusInfoCallback;
+		FTabletContextContainer TabletContexts;
+		TArray<IStylusInputEventHandler*> EventHandlers;
+	};
+}
diff --git a/Engine/Plugins/Editor/StylusInput/Source/StylusInputRealTimeStylus/Private/RealTimeStylusPluginSync.cpp b/Engine/Plugins/Editor/StylusInput/Source/StylusInputRealTimeStylus/Private/RealTimeStylusPluginSync.cpp
new file mode 100644
--- /dev/null
+++ b/Engine/Plugins/Editor/StylusInput/Source/StylusInputRealTimeStylus/Private/RealTimeStylusPluginSync.cpp
@@ -0,0 +1,131 @@
+﻿// Copyright Epic Games, Inc. All Rights Reserved.
+
+#include "RealTimeStylusPluginSync.h"
+
+#include "RealTimeStylusAPI.h"
+
+namespace UE::StylusInput::RealTimeStylus
+{
+	FRealTimeStylusPluginSync::FRealTimeStylusPluginSync(IStylusInputInstance* Instance,
+	                                                     FGetWindowContextCallback&& GetWindowContext,
+	                                                     FUpdateTabletContextsCallback&& UpdateTabletContextsCallback,
+	                                                     FUpdateStylusInfoCallback UpdateStylusInfoCallback,
+	                                                     IStylusInputEventHandler* EventHandler)
+		: FRealTimeStylusPluginBase(Instance, MoveTemp(GetWindowContext), MoveTemp(UpdateTabletContextsCallback), MoveTemp(UpdateStylusInfoCallback))
+	{
+		if (EventHandler)
+		{
+			// Immediately install an event handler during construction to capture events coming through during plugin initialization.
+			AddEventHandler(EventHandler);
+		}
+	}
+
+	FRealTimeStylusPluginSync::~FRealTimeStylusPluginSync()
+	{
+		if (FreeThreadedMarshaler)
+		{
+			FreeThreadedMarshaler->Release();
+		}
+	}
+
+	HRESULT FRealTimeStylusPluginSync::CreateFreeThreadMarshaler()
+	{
+		check(FreeThreadedMarshaler == nullptr);
+
+		const FRealTimeStylusAPI& RealTimeStylusAPI = FRealTimeStylusAPI::GetInstance();
+		return RealTimeStylusAPI.CoCreateFreeThreadedMarshaler(this, &FreeThreadedMarshaler);
+	}
+
+	HRESULT FRealTimeStylusPluginSync::RealTimeStylusEnabled(IRealTimeStylus* RealTimeStylus, const ULONG TabletContextIDsCount,
+	                                                             const TABLET_CONTEXT_ID* TabletContextIDs)
+	{
+		return ProcessRealTimeStylusEnabled(RealTimeStylus, TabletContextIDsCount, TabletContextIDs);
+	}
+
+	HRESULT FRealTimeStylusPluginSync::RealTimeStylusDisabled(IRealTimeStylus* RealTimeStylus, const ULONG TabletContextIDsCount,
+	                                                              const TABLET_CONTEXT_ID* TabletContextIDs)
+	{
+		return ProcessRealTimeStylusDisabled(RealTimeStylus, TabletContextIDsCount, TabletContextIDs);
+	}
+
+	HRESULT FRealTimeStylusPluginSync::StylusDown(IRealTimeStylus* RealTimeStylus, const StylusInfo* StylusInfo, const ULONG PropertyCount,
+	                                                  LONG* PacketBuffer, LONG**)
+	{
+		return ProcessPackets(StylusInfo, 1, PropertyCount, EPacketType::StylusDown, reinterpret_cast<int32*>(PacketBuffer));
+	}
+
+	HRESULT FRealTimeStylusPluginSync::StylusUp(IRealTimeStylus* RealTimeStylus, const StylusInfo* StylusInfo, const ULONG PropertyCount,
+	                                                LONG* PacketBuffer, LONG**)
+	{
+		return ProcessPackets(StylusInfo, 1, PropertyCount, EPacketType::StylusUp, reinterpret_cast<int32*>(PacketBuffer));
+	}
+
+	HRESULT FRealTimeStylusPluginSync::TabletAdded(IRealTimeStylus* RealTimeStylus, IInkTablet* Tablet)
+	{
+		return ProcessTabletAdded(Tablet);
+	}
+
+	HRESULT FRealTimeStylusPluginSync::TabletRemoved(IRealTimeStylus* RealTimeStylus, const LONG TabletIndex)
+	{
+		return ProcessTabletRemoved(TabletIndex);
+	}
+
+	HRESULT FRealTimeStylusPluginSync::InAirPackets(IRealTimeStylus* RealTimeStylus, const StylusInfo* StylusInfo, const ULONG PacketCount,
+	                                                    const ULONG PacketBufferLength, LONG* PacketBuffer, ULONG*, LONG**)
+	{
+		return ProcessPackets(StylusInfo, PacketCount, PacketBufferLength, EPacketType::AboveDigitizer, reinterpret_cast<int32*>(PacketBuffer));
+	}
+
+	HRESULT FRealTimeStylusPluginSync::Packets(IRealTimeStylus* RealTimeStylus, const StylusInfo* StylusInfo, const ULONG PacketCount,
+	                                               const ULONG PacketBufferLength, LONG* PacketBuffer, ULONG*, LONG**)
+	{
+		return ProcessPackets(StylusInfo, PacketCount, PacketBufferLength, EPacketType::OnDigitizer, reinterpret_cast<int32*>(PacketBuffer));
+	}
+
+	HRESULT FRealTimeStylusPluginSync::StylusInRange(IRealTimeStylus* RealTimeStylus, TABLET_CONTEXT_ID TabletContextID, STYLUS_ID StylusID)
+	{
+		return ProcessStylusInRange(RealTimeStylus, TabletContextID, StylusID);
+	}
+
+	HRESULT FRealTimeStylusPluginSync::StylusOutOfRange(IRealTimeStylus* RealTimeStylus, TABLET_CONTEXT_ID TabletContextID, STYLUS_ID StylusID)
+	{
+		return E_NOTIMPL;
+	}
+
+	HRESULT FRealTimeStylusPluginSync::StylusButtonDown(IRealTimeStylus* RealTimeStylus, STYLUS_ID StylusID, const GUID* GuidStylusButton, POINT* StylusPos)
+	{
+		return E_NOTIMPL;
+	}
+
+	HRESULT FRealTimeStylusPluginSync::StylusButtonUp(IRealTimeStylus* RealTimeStylus, STYLUS_ID StylusID, const GUID* GuidStylusButton, POINT* StylusPos)
+	{
+		return E_NOTIMPL;
+	}
+
+	HRESULT FRealTimeStylusPluginSync::CustomStylusDataAdded(IRealTimeStylus* RealTimeStylus, const GUID* GuidId, ULONG DataCount, const BYTE* Data)
+	{
+		return E_NOTIMPL;
+	}
+
+	HRESULT FRealTimeStylusPluginSync::SystemEvent(IRealTimeStylus* RealTimeStylus, TABLET_CONTEXT_ID TabletContextID, STYLUS_ID StylusID,
+	                                                   SYSTEM_EVENT Event, SYSTEM_EVENT_DATA EventData)
+	{
+		return E_NOTIMPL;
+	}
+
+	HRESULT FRealTimeStylusPluginSync::Error(IRealTimeStylus*, IStylusPlugin* Plugin, const RealTimeStylusDataInterest DataInterest,
+	                                             const HRESULT ErrorCode, LONG_PTR*)
+	{
+		return ProcessError(DataInterest, ErrorCode);
+	}
+
+	HRESULT FRealTimeStylusPluginSync::UpdateMapping(IRealTimeStylus* RealTimeStylus)
+	{
+		return E_NOTIMPL;
+	}
+
+	HRESULT FRealTimeStylusPluginSync::DataInterest(RealTimeStylusDataInterest* DataInterest)
+	{
+		return ProcessDataInterest(DataInterest);
+	}
+}
diff --git a/Engine/Plugins/Editor/StylusInput/Source/StylusInputRealTimeStylus/Private/RealTimeStylusPluginSync.h b/Engine/Plugins/Editor/StylusInput/Source/StylusInputRealTimeStylus/Private/RealTimeStylusPluginSync.h
new file mode 100644
--- /dev/null
+++ b/Engine/Plugins/Editor/StylusInput/Source/StylusInputRealTimeStylus/Private/RealTimeStylusPluginSync.h
@@ -0,0 +1,80 @@
+﻿// Copyright Epic Games, Inc. All Rights Reserved.
+
+#pragma once
+
+#include "RealTimeStylusPluginBase.h"
+
+namespace UE::StylusInput::RealTimeStylus
+{
+	class FRealTimeStylusPluginSync final : public IStylusSyncPlugin, public FRealTimeStylusPluginBase
+	{
+	public:
+		FRealTimeStylusPluginSync(IStylusInputInstance* Instance, FGetWindowContextCallback&& GetWindowContext,
+		                          FUpdateTabletContextsCallback&& UpdateTabletContextsCallback, FUpdateStylusInfoCallback UpdateStylusInfoCallback,
+		                          IStylusInputEventHandler* EventHandler);
+		~FRealTimeStylusPluginSync();
+
+		HRESULT CreateFreeThreadMarshaler();
+
+		// IUnknown
+		virtual HRESULT QueryInterface(REFIID InterfaceID, void** InterfaceObject) override
+		{
+			if (InterfaceID == IID_IStylusSyncPlugin || InterfaceID == IID_IUnknown)
+			{
+				*InterfaceObject = this;
+				AddRef();
+				return S_OK;
+			}
+
+			if (InterfaceID == IID_IMarshal && FreeThreadedMarshaler)
+			{
+				return FreeThreadedMarshaler->QueryInterface(InterfaceID, InterfaceObject);
+			}
+
+			*InterfaceObject = nullptr;
+			return E_NOINTERFACE;
+		}
+
+		virtual ULONG AddRef() override
+		{
+			return ++RefCount;
+		}
+
+		virtual ULONG Release() override
+		{
+			const int32 NewRefCount = --RefCount;
+			if (NewRefCount == 0)
+				delete this;
+
+			return NewRefCount;
+		}
+
+		// IStylusPlugin
+		virtual HRESULT RealTimeStylusEnabled(IRealTimeStylus* RealTimeStylus, ULONG TabletContextIDsCount, const TABLET_CONTEXT_ID* TabletContextIDs) override;
+		virtual HRESULT RealTimeStylusDisabled(IRealTimeStylus* RealTimeStylus, ULONG TabletContextIDsCount, const TABLET_CONTEXT_ID* TabletContextIDs) override;
+		virtual HRESULT StylusDown(IRealTimeStylus* RealTimeStylus, const StylusInfo* StylusInfo, ULONG PropertyCount, LONG* PacketBuffer, LONG** InOutPkt) override;
+		virtual HRESULT StylusUp(IRealTimeStylus* RealTimeStylus, const StylusInfo* StylusInfo, ULONG PropertyCount, LONG* PacketBuffer, LONG** InOutPkt) override;
+		virtual HRESULT TabletAdded(IRealTimeStylus* RealTimeStylus, IInkTablet* Tablet) override;
+		virtual HRESULT TabletRemoved(IRealTimeStylus* RealTimeStylus, LONG TabletIndex) override;
+		virtual HRESULT InAirPackets(IRealTimeStylus* RealTimeStylus, const StylusInfo* StylusInfo, ULONG PacketCount, ULONG PacketBufferLength, LONG* PacketBuffer, ULONG* InOutPacketCount, LONG** InOutPacketBuffer) override;
+		virtual HRESULT Packets(IRealTimeStylus* RealTimeStylus, const StylusInfo* StylusInfo, ULONG PacketCount, ULONG PacketBufferLength, LONG* PacketBuffer, ULONG* InOutPacketCount, LONG** InOutPacketBuffer) override;
+		virtual HRESULT StylusInRange(IRealTimeStylus* RealTimeStylus, TABLET_CONTEXT_ID TabletContextID, STYLUS_ID StylusID) override;
+		virtual HRESULT StylusOutOfRange(IRealTimeStylus* RealTimeStylus, TABLET_CONTEXT_ID TabletContextID, STYLUS_ID StylusID) override;
+		virtual HRESULT StylusButtonDown(IRealTimeStylus* RealTimeStylus, STYLUS_ID StylusID, const GUID* GuidStylusButton, POINT* StylusPos) override;
+		virtual HRESULT StylusButtonUp(IRealTimeStylus* RealTimeStylus, STYLUS_ID StylusID, const GUID* GuidStylusButton, POINT* StylusPos) override;
+		virtual HRESULT CustomStylusDataAdded(IRealTimeStylus* RealTimeStylus, const GUID* GuidId, ULONG DataCount, const BYTE* Data) override;
+		virtual HRESULT SystemEvent(IRealTimeStylus* RealTimeStylus, TABLET_CONTEXT_ID TabletContextID, STYLUS_ID StylusID, SYSTEM_EVENT Event, SYSTEM_EVENT_DATA EventData) override;
+		virtual HRESULT Error(IRealTimeStylus* RealTimeStylus, IStylusPlugin* Plugin, RealTimeStylusDataInterest DataInterest, HRESULT ErrorCode, LONG_PTR* InternalKey) override;
+		virtual HRESULT UpdateMapping(IRealTimeStylus* RealTimeStylus) override;
+		virtual HRESULT DataInterest(RealTimeStylusDataInterest* DataInterest) override;
+
+	protected:
+		// FRealTimeStylusPluginBase
+		virtual FString GetName() const override { return "Asynchronous"; }
+
+	private:
+		int32 RefCount = 1;
+
+		IUnknown* FreeThreadedMarshaler = nullptr;
+	};
+}
diff --git a/Engine/Plugins/Editor/StylusInput/Source/StylusInputRealTimeStylus/Private/RealTimeStylusStats.h b/Engine/Plugins/Editor/StylusInput/Source/StylusInputRealTimeStylus/Private/RealTimeStylusStats.h
new file mode 100644
--- /dev/null
+++ b/Engine/Plugins/Editor/StylusInput/Source/StylusInputRealTimeStylus/Private/RealTimeStylusStats.h
@@ -0,0 +1,40 @@
+﻿// Copyright Epic Games, Inc. All Rights Reserved.
+
+#pragma once
+
+#include <Windows/WindowsPlatformTime.h>
+
+namespace UE::StylusInput::RealTimeStylus
+{
+	class FPacketStats
+	{
+	public:
+		void NewPacket()
+		{
+			const uint32 Timestamp = FPlatformTime::Cycles();
+
+			if (Latest - Earliest >= CyclesPerSecond)
+			{
+				PacketsPerSecond = Earliest != Latest ? NumPackets / ((Latest - Earliest) / CyclesPerSecond) : 0.0f;
+				Earliest = Timestamp;
+				NumPackets = -1;
+			}
+
+			Latest = Timestamp;
+			++NumPackets;
+		}
+
+		float GetPacketsPerSecond() const
+		{
+			const uint32 Timestamp = FPlatformTime::Cycles();
+			return Timestamp - Latest <= CyclesPerSecond ? PacketsPerSecond : 0.0f;
+		}
+
+	private:
+		uint32 Earliest = 0;
+		uint32 Latest = 0;
+		uint32 NumPackets = 0;
+		float PacketsPerSecond = 0.0f;
+		const double CyclesPerSecond = 1.0 / FPlatformTime::GetSecondsPerCycle();
+	};
+}
diff --git a/Engine/Plugins/Editor/StylusInput/Source/StylusInputRealTimeStylus/Private/RealTimeStylusTabletContext.cpp b/Engine/Plugins/Editor/StylusInput/Source/StylusInputRealTimeStylus/Private/RealTimeStylusTabletContext.cpp
new file mode 100644
--- /dev/null
+++ b/Engine/Plugins/Editor/StylusInput/Source/StylusInputRealTimeStylus/Private/RealTimeStylusTabletContext.cpp
@@ -0,0 +1,34 @@
+﻿// Copyright Epic Games, Inc. All Rights Reserved.
+
+#include "RealTimeStylusTabletContext.h"
+
+#include <Microsoft/COMPointer.h>
+#include <msinkaut.h>
+
+namespace UE::StylusInput::RealTimeStylus
+{
+	FPacketPropertyConstant PacketPropertyConstants[] = {
+		{EPacketPropertyType::X,                       GUID_PACKETPROPERTY_GUID_X,                        STR_GUID_X                      },
+		{EPacketPropertyType::Y,                       GUID_PACKETPROPERTY_GUID_Y,                        STR_GUID_Y                      },
+		{EPacketPropertyType::Z,                       GUID_PACKETPROPERTY_GUID_Z,                        STR_GUID_Z                      },
+		{EPacketPropertyType::PacketStatus,            GUID_PACKETPROPERTY_GUID_PACKET_STATUS,            STR_GUID_PAKETSTATUS            },
+		{EPacketPropertyType::TimerTick,               GUID_PACKETPROPERTY_GUID_TIMER_TICK,               STR_GUID_TIMERTICK              },
+		{EPacketPropertyType::SerialNumber,            GUID_PACKETPROPERTY_GUID_SERIAL_NUMBER,            STR_GUID_SERIALNUMBER           },
+		{EPacketPropertyType::NormalPressure,          GUID_PACKETPROPERTY_GUID_NORMAL_PRESSURE,          STR_GUID_NORMALPRESSURE         },
+		{EPacketPropertyType::TangentPressure,         GUID_PACKETPROPERTY_GUID_TANGENT_PRESSURE,         STR_GUID_TANGENTPRESSURE        },
+		{EPacketPropertyType::ButtonPressure,          GUID_PACKETPROPERTY_GUID_BUTTON_PRESSURE,          STR_GUID_BUTTONPRESSURE         },
+		{EPacketPropertyType::XTiltOrientation,        GUID_PACKETPROPERTY_GUID_X_TILT_ORIENTATION,       STR_GUID_XTILTORIENTATION       },
+		{EPacketPropertyType::YTiltOrientation,        GUID_PACKETPROPERTY_GUID_Y_TILT_ORIENTATION,       STR_GUID_YTILTORIENTATION       },
+		{EPacketPropertyType::AzimuthOrientation,      GUID_PACKETPROPERTY_GUID_AZIMUTH_ORIENTATION,      STR_GUID_AZIMUTHORIENTATION     },
+		{EPacketPropertyType::AltitudeOrientation,     GUID_PACKETPROPERTY_GUID_ALTITUDE_ORIENTATION,     STR_GUID_ALTITUDEORIENTATION    },
+		{EPacketPropertyType::TwistOrientation,        GUID_PACKETPROPERTY_GUID_TWIST_ORIENTATION,        STR_GUID_TWISTORIENTATION       },
+		{EPacketPropertyType::PitchRotation,           GUID_PACKETPROPERTY_GUID_PITCH_ROTATION,           STR_GUID_PITCHROTATION          },
+		{EPacketPropertyType::RollRotation,            GUID_PACKETPROPERTY_GUID_ROLL_ROTATION,            STR_GUID_ROLLROTATION           },
+		{EPacketPropertyType::YawRotation,             GUID_PACKETPROPERTY_GUID_YAW_ROTATION,             STR_GUID_YAWROTATION            },
+		{EPacketPropertyType::Width,                   GUID_PACKETPROPERTY_GUID_WIDTH,                    STR_GUID_WIDTH                  },
+		{EPacketPropertyType::Height,                  GUID_PACKETPROPERTY_GUID_HEIGHT,                   STR_GUID_HEIGHT                 },
+		{EPacketPropertyType::FingerContactConfidence, GUID_PACKETPROPERTY_GUID_FINGERCONTACTCONFIDENCE,  STR_GUID_FINGERCONTACTCONFIDENCE},
+		{EPacketPropertyType::DeviceContactID,         GUID_PACKETPROPERTY_GUID_DEVICE_CONTACT_ID,        STR_GUID_DEVICE_CONTACT_ID      }
+	};
+	static_assert(std::size(PacketPropertyConstants) == static_cast<uint32>(EPacketPropertyType::Num_Enumerators));
+}
diff --git a/Engine/Plugins/Editor/StylusInput/Source/StylusInputRealTimeStylus/Private/RealTimeStylusTabletContext.h b/Engine/Plugins/Editor/StylusInput/Source/StylusInputRealTimeStylus/Private/RealTimeStylusTabletContext.h
new file mode 100644
--- /dev/null
+++ b/Engine/Plugins/Editor/StylusInput/Source/StylusInputRealTimeStylus/Private/RealTimeStylusTabletContext.h
@@ -0,0 +1,163 @@
+﻿// Copyright Epic Games, Inc. All Rights Reserved.
+
+#pragma once
+
+#include <StylusInputPacket.h>
+#include <StylusInputTabletContext.h>
+#include <StylusInputUtils.h>
+#include <Math/IntRect.h>
+#include <Math/Vector.h>
+
+#include "RealTimeStylusUtils.h"
+
+// The following includes are needed for 'GUID' type.
+#include <Windows/AllowWindowsPlatformTypes.h>
+#include <Windows/HideWindowsPlatformTypes.h>
+
+namespace UE::StylusInput::RealTimeStylus
+{
+	enum class EPacketPropertyType : int8
+	{
+		X = 0, // The x-coordinate in the tablet coordinate space. Each packet contains this property by default. The origin (0,0) of the tablet is the upper-left corner.
+		Y, // The y-coordinate in the tablet coordinate space. Each packet contains this property by default. The origin (0,0) of the tablet is the upper-left corner.
+		Z, // The z-coordinate or distance of the pen tip from the tablet surface. The TabletPropertyMetricUnit enumeration type determines the unit of measurement for this property.
+		PacketStatus, // Contains one or more of the following flag values: The cursor is touching the drawing surface (Value = 1). The cursor is inverted. For example, the eraser end of the pen is pointing toward the surface (Value = 2). Not used (Value = 4). The barrel button is pressed (Value = 8).
+		TimerTick, // The time the packet was generated in milliseconds since system start.
+		SerialNumber, // The packet property for identifying the packet. This is the same value you use to retrieve the packet from the packet queue.
+		NormalPressure, // The pressure of the pen tip perpendicular to the tablet surface. The greater the pressure on the pen tip, the more ink that is drawn.
+		TangentPressure, // The pressure of the pen tip along the plane of the tablet surface.
+		ButtonPressure, // The pressure on a pressure sensitive button.
+		XTiltOrientation, // The angle between the y,z-plane and the pen and y-axis plane. Applies to a pen cursor. The value is 0 when the pen is perpendicular to the drawing surface and is positive when the pen is to the right of perpendicular.
+		YTiltOrientation, // The angle between the x,z-plane and the pen and x-axis plane. Applies to a pen cursor. The value is 0 when the pen is perpendicular to the drawing surface and is positive when the pen is upward or away from the user.
+		AzimuthOrientation, // The clockwise rotation of the cursor about the z-axis through a full circular range.
+		AltitudeOrientation, // The angle between the axis of the pen and the surface of the tablet. The value is 0 when the pen is parallel to the surface and 90 when the pen is perpendicular to the surface. The values are negative when the pen is inverted.
+		TwistOrientation, // The clockwise rotation of the cursor about its own axis.
+		PitchRotation, // The packet property that indicates whether the tip is above or below a horizontal line that is perpendicular to the writing surface. Note: This requires a 3D digitizer. The value is positive if the tip is above the line and negative if it is below the line. For example, if you hold the pen in front of you and write on an imaginary wall, the pitch is positive if the tip is above a line extending from you to the wall.
+		RollRotation, // The clockwise rotation of the pen around its own axis. Note: This requires a 3D digitizer.
+		YawRotation, // The angle of the pen to the left or right around the center of its horizontal axis when the pen is horizontal. Note: This requires a 3D digitizer. If you hold the pen in front of you and write on an imaginary wall, zero yaw indicates that the pen is perpendicular to the wall. The value is negative if the tip is to the left of perpendicular and positive if the tip is to the right of perpendicular.
+		Width, // The width of the contact area on a touch digitizer.
+		Height, // The height of the contact area on a touch digitizer.
+		FingerContactConfidence, // The level of confidence that there was finger contact on a touch digitizer.
+		DeviceContactID, // The device contact identifier for a packet.
+
+		Num_Enumerators, // THIS IS NOT A VALID ENUMERATOR, BUT IT CAN BE USED TO QUERY HOW MANY PACKET PROPERTIES THERE ARE.
+		Invalid_Enumerator = INDEX_NONE // THIS IS NOT A VALID ENUMERATOR, BUT IT CAN BE USED TO INDICATE THAT A PACKET PROPERTY ENTRY IS INVALID.
+	};
+
+	enum class ETabletPropertyMetricUnit : int8
+	{
+		Default, // Units are unknown.
+		Inches,
+		Centimeters,
+		Degrees,
+		Radians,
+		Seconds,
+		Pounds,
+		Grams,
+
+		Num_Enumerators, // THIS IS NOT A VALID ENUMERATOR, BUT IT CAN BE USED TO QUERY HOW MANY PACKET PROPERTIES THERE ARE.
+		Invalid_Enumerator = INDEX_NONE // THIS IS NOT A VALID ENUMERATOR, BUT IT CAN BE USED TO INDICATE THAT A PACKET PROPERTY ENTRY IS INVALID.
+	};
+
+	struct FPacketProperty
+	{
+		EPacketPropertyType Type = EPacketPropertyType::Invalid_Enumerator;
+		ETabletPropertyMetricUnit MetricUnit = ETabletPropertyMetricUnit::Default;
+		int32 Minimum = 0;
+		int32 Maximum = 0;
+		float Resolution = 0.0f;
+	};
+
+	struct FPacketPropertyHandler
+	{
+		using FFuncSetProperty = void(FStylusInputPacket&, int32, const int8*);
+		FFuncSetProperty* SetProperty = nullptr;
+
+		static constexpr int32 SetPropertyDataBufferLength = 24;
+		int8 SetPropertyData[SetPropertyDataBufferLength];
+	};
+
+	struct FPacketPropertyConstant
+	{
+		EPacketPropertyType PacketPropertyType;
+		GUID Guid;
+		const wchar_t* StrGuid;
+	};
+
+	extern FPacketPropertyConstant PacketPropertyConstants[static_cast<uint32>(EPacketPropertyType::Num_Enumerators)];
+
+	struct FWindowContext
+	{
+		FVector2d XYScale;
+		FVector2d XYMaximum;
+		FVector2d WindowSize;
+	};
+
+	class FTabletContext final : public IStylusInputTabletContext
+	{
+	public:
+		explicit FTabletContext(const uint32 ID)
+			: ID(ID)
+		{}
+
+		virtual uint32 GetID() const override { return ID; }
+
+		virtual FString GetName() const override
+		{
+			if (!Name.IsEmpty() && !PlugAndPlayID.IsEmpty())
+			{
+				return FString::Format(TEXT("{0} ({1})"), {Name, PlugAndPlayID});
+			}
+			return !Name.IsEmpty() ? Name : PlugAndPlayID;
+		}
+
+		virtual FIntRect GetInputRectangle() const override { return InputRectangle; }
+		virtual ETabletHardwareCapabilities GetHardwareCapabilities() const override { return HardwareCapabilities; }
+		virtual ETabletSupportedProperties GetSupportedProperties() const override { return SupportedProperties; }
+
+		uint32 ID = 0;
+		FString Name;
+		FString PlugAndPlayID;
+		FIntRect InputRectangle;
+		ETabletSupportedProperties SupportedProperties = ETabletSupportedProperties::None;
+		ETabletHardwareCapabilities HardwareCapabilities = ETabletHardwareCapabilities::None;
+		FPacketProperty PacketProperties[static_cast<uint32>(EPacketPropertyType::Num_Enumerators)] = {};
+		FPacketPropertyHandler PacketPropertyHandlers[static_cast<uint32>(EPacketPropertyType::Num_Enumerators)] = {};
+	};
+
+	using FTabletContextContainer = Private::TSharedRefDataContainer<FTabletContext, false>;
+	using FTabletContextThreadSafeContainer = Private::TSharedRefDataContainer<FTabletContext, true>;
+
+	class FStylusButton final : public IStylusInputStylusButton
+	{
+	public:
+		virtual FString GetID() const override { return ID; }
+		virtual FString GetName() const override { return Name; }
+
+		FString ID;
+		FString Name;
+	};
+
+	class FStylusInfo final : public IStylusInputStylusInfo
+	{
+	public:
+		explicit FStylusInfo(const uint32 ID)
+			: ID(ID)
+		{}
+
+		virtual uint32 GetID() const override { return ID; }
+		virtual FString GetName() const override { return Name; }
+		virtual uint32 GetNumButtons() const override { return Buttons.Num(); }
+
+		virtual const IStylusInputStylusButton* GetButton(int32 Index) const override
+		{
+			return 0 <= Index && Index < Buttons.Num() ? reinterpret_cast<const IStylusInputStylusButton*>(&Buttons[Index]) : nullptr;
+		}
+
+		uint32 ID = 0;
+		FString Name;
+		TArray<FStylusButton> Buttons;
+	};
+
+	using FStylusInfoThreadSafeContainer = Private::TSharedRefDataContainer<FStylusInfo, true>;
+}
diff --git a/Engine/Plugins/Editor/StylusInput/Source/StylusInputRealTimeStylus/Private/RealTimeStylusUtils.cpp b/Engine/Plugins/Editor/StylusInput/Source/StylusInputRealTimeStylus/Private/RealTimeStylusUtils.cpp
new file mode 100644
--- /dev/null
+++ b/Engine/Plugins/Editor/StylusInput/Source/StylusInputRealTimeStylus/Private/RealTimeStylusUtils.cpp
@@ -0,0 +1,26 @@
+﻿// Copyright Epic Games, Inc. All Rights Reserved.
+
+#include "RealTimeStylusUtils.h"
+
+#include <Windows/AllowWindowsPlatformAtomics.h>
+	#include <comdef.h>
+#include <Windows/HideWindowsPlatformAtomics.h>
+
+#include <StylusInputUtils.h>
+
+using namespace UE::StylusInput::Private;
+
+namespace UE::StylusInput::RealTimeStylus
+{
+	FString GetCOMErrorMessage(const HRESULT Result)
+	{
+		const _com_error COMError(Result);
+		const LPCTSTR ErrorMessage = COMError.ErrorMessage();
+		return ErrorMessage;
+	}
+
+	void LogCOMError(const FString& Preamble, const HRESULT Result)
+	{
+		LogError(Preamble, FString::Format(TEXT("COM Interface Error: {0}"), {GetCOMErrorMessage(Result)}));
+	}
+}
diff --git a/Engine/Plugins/Editor/StylusInput/Source/StylusInputRealTimeStylus/Private/RealTimeStylusUtils.h b/Engine/Plugins/Editor/StylusInput/Source/StylusInputRealTimeStylus/Private/RealTimeStylusUtils.h
new file mode 100644
--- /dev/null
+++ b/Engine/Plugins/Editor/StylusInput/Source/StylusInputRealTimeStylus/Private/RealTimeStylusUtils.h
@@ -0,0 +1,30 @@
+﻿// Copyright Epic Games, Inc. All Rights Reserved.
+
+#pragma once
+
+#include <Containers/UnrealString.h>
+
+namespace UE::StylusInput::RealTimeStylus
+{
+	void LogCOMError(const FString& Preamble, HRESULT Result);
+
+	inline bool Succeeded(const HRESULT Result, const FString& LogPreamble)
+	{
+		if (Result >= 0)
+		{
+			return true;
+		}
+		LogCOMError(LogPreamble, Result);
+		return false;
+	}
+
+	inline bool Failed(const HRESULT Result, const FString& LogPreamble)
+	{
+		if (Result >= 0)
+		{
+			return false;
+		}
+		LogCOMError(LogPreamble, Result);
+		return true;
+	}
+}
diff --git a/Engine/Plugins/Editor/StylusInput/Source/StylusInputRealTimeStylus/Private/StylusInputRealTimeStylusModule.cpp b/Engine/Plugins/Editor/StylusInput/Source/StylusInputRealTimeStylus/Private/StylusInputRealTimeStylusModule.cpp
new file mode 100644
--- /dev/null
+++ b/Engine/Plugins/Editor/StylusInput/Source/StylusInputRealTimeStylus/Private/StylusInputRealTimeStylusModule.cpp
@@ -0,0 +1,37 @@
+﻿// Copyright Epic Games, Inc. All Rights Reserved.
+
+#include <StylusInput.h>
+#include <Modules/ModuleManager.h>
+
+#include "RealTimeStylusInterface.h"
+
+namespace UE::StylusInput::RealTimeStylus
+{
+	class FStylusInputRealTimeStylusModule : public IModuleInterface
+	{
+	public:
+		virtual void StartupModule() override;
+		virtual void ShutdownModule() override;
+
+	private:
+		TUniquePtr<IStylusInputInterface> RegisteredInterface;
+	};
+
+	void FStylusInputRealTimeStylusModule::StartupModule()
+	{
+		RegisteredInterface = FRealTimeStylusInterface::Create();
+		RegisterInterface(RegisteredInterface.Get());
+	}
+
+	void FStylusInputRealTimeStylusModule::ShutdownModule()
+	{
+		UnregisterInterface(RegisteredInterface.Get());
+	}
+}
+
+IMPLEMENT_MODULE(UE::StylusInput::RealTimeStylus::FStylusInputRealTimeStylusModule, StylusInputRealTimeStylus);
+
+/*
+ * TODO
+ * [ ] Propagate buttons data for StylusButtonDown/StylusButtonDown
+ */
diff --git a/Engine/Plugins/Editor/StylusInput/Source/StylusInputRealTimeStylus/StylusInputRealTimeStylus.Build.cs b/Engine/Plugins/Editor/StylusInput/Source/StylusInputRealTimeStylus/StylusInputRealTimeStylus.Build.cs
new file mode 100644
--- /dev/null
+++ b/Engine/Plugins/Editor/StylusInput/Source/StylusInputRealTimeStylus/StylusInputRealTimeStylus.Build.cs
@@ -0,0 +1,19 @@
+﻿// Copyright Epic Games, Inc. All Rights Reserved.
+
+namespace UnrealBuildTool.Rules
+{
+	public class StylusInputRealTimeStylus : ModuleRules
+	{
+		public StylusInputRealTimeStylus(ReadOnlyTargetRules Target) : base(Target)
+		{
+			PrivateDependencyModuleNames.AddRange(
+				[
+					"Core",
+					"Slate",
+					"SlateCore",
+					"StylusInput"
+				]
+			);
+		}
+	}
+}
\ No newline at end of file
diff --git a/Engine/Plugins/Editor/StylusInput/Source/StylusInputWintab/Private/StylusInputWintabModule.cpp b/Engine/Plugins/Editor/StylusInput/Source/StylusInputWintab/Private/StylusInputWintabModule.cpp
new file mode 100644
--- /dev/null
+++ b/Engine/Plugins/Editor/StylusInput/Source/StylusInputWintab/Private/StylusInputWintabModule.cpp
@@ -0,0 +1,43 @@
+﻿// Copyright Epic Games, Inc. All Rights Reserved.
+
+#include <StylusInput.h>
+#include <Modules/ModuleManager.h>
+
+#include "WintabInterface.h"
+
+namespace UE::StylusInput::Wintab
+{
+	class FStylusInputWintabModule : public IModuleInterface
+	{
+	public:
+		virtual void StartupModule() override;
+		virtual void ShutdownModule() override;
+
+	private:
+		TUniquePtr<IStylusInputInterface> RegisteredInterface;
+	};
+
+	void FStylusInputWintabModule::StartupModule()
+	{
+		RegisteredInterface = FWintabInterface::Create();
+		RegisterInterface(RegisteredInterface.Get());
+	}
+
+	void FStylusInputWintabModule::ShutdownModule()
+	{
+		UnregisterInterface(RegisteredInterface.Get());
+	}
+}
+
+IMPLEMENT_MODULE(UE::StylusInput::Wintab::FStylusInputWintabModule, StylusInputWintab);
+
+/*
+ * TODO
+ * [ ] Make Wintab the default (with additional editor settings?)
+ * [ ] Propagate buttons data for both Wintab and Windows (PK_BUTTONS and/or StylusButtonDown/StylusButtonDown events)
+ * [ ] Queue size
+ * [ ] Query multiple packets at once
+ * [ ] Async
+ * [ ] Deal with cursors not having a physical ID
+ * [ ] Allegedly, some Wintab libraries don't handle relative button presses, so maybe switch to absolute?
+ */
diff --git a/Engine/Plugins/Editor/StylusInput/Source/StylusInputWintab/Private/WintabAPI.cpp b/Engine/Plugins/Editor/StylusInput/Source/StylusInputWintab/Private/WintabAPI.cpp
new file mode 100644
--- /dev/null
+++ b/Engine/Plugins/Editor/StylusInput/Source/StylusInputWintab/Private/WintabAPI.cpp
@@ -0,0 +1,167 @@
+﻿// Copyright Epic Games, Inc. All Rights Reserved.
+
+#include "WintabAPI.h"
+
+#include <StylusInputUtils.h>
+#include <Windows/WindowsPlatformProcess.h>
+
+#define LOG_PREAMBLE "WintabAPI"
+
+using namespace UE::StylusInput::Private;
+
+namespace UE::StylusInput::Wintab
+{
+	FWintabAPI::FWintabAPI()
+	{
+		auto GetDllHandle = [&DllHandles = DllHandles](const TCHAR* DllName) -> void* {
+			void* DllHandle = FWindowsPlatformProcess::GetDllHandle(DllName);
+			if (DllHandle)
+			{
+				DllHandles.Add(DllHandle);
+				LogVerbose(LOG_PREAMBLE, FString::Format(TEXT("Retrieved DLL handle for {0}."), {DllName}));
+			}
+			else
+			{
+				LogError(LOG_PREAMBLE, FString::Format(TEXT("Could not get DLL handle for {0}."), {DllName}));
+			}
+			return DllHandle;
+		};
+
+		auto GetDllExport = []<typename FFuncType>(const TCHAR* DllName, void* DllHandle, const TCHAR* ExportName, FFuncType& ExportHandle)
+		{
+			void* FuncPtr = FWindowsPlatformProcess::GetDllExport(DllHandle, ExportName);
+			if (FuncPtr)
+			{
+				LogVerbose(LOG_PREAMBLE, FString::Format(TEXT("Retrieved DLL export '{0}' from {1}."), {ExportName, DllName}));
+			}
+			else
+			{
+				LogError(LOG_PREAMBLE, FString::Format(TEXT("Could not get DLL export '{0}' in {1}."), {ExportName, DllName}));
+			}
+			ExportHandle = static_cast<FFuncType>(FuncPtr);
+		};
+
+		const TCHAR* User32DllName = TEXT("User32.dll");
+		void* User32DllHandle = GetDllHandle(User32DllName);
+		GetDllExport(User32DllName, User32DllHandle, TEXT("GetWindowRect"), GetWindowRect);
+		GetDllExport(User32DllName, User32DllHandle, TEXT("GetSystemMetrics"), GetSystemMetrics);
+
+		const TCHAR* Wintab32DllName = TEXT("Wintab32.dll");
+		void* Wintab32DllHandle = GetDllHandle(Wintab32DllName);
+		GetDllExport(Wintab32DllName, Wintab32DllHandle, TEXT("WTInfoW"), WtInfo);
+		GetDllExport(Wintab32DllName, Wintab32DllHandle, TEXT("WTOpenW"), WtOpen);
+		GetDllExport(Wintab32DllName, Wintab32DllHandle, TEXT("WTClose"), WtClose);
+		GetDllExport(Wintab32DllName, Wintab32DllHandle, TEXT("WTPacket"), WtPacket);
+
+		if (WtInfo)
+		{
+			bWintabIsAvailable = WtInfo(0, 0, nullptr) != 0;
+			if (!bWintabIsAvailable)
+			{
+				LogVerbose(LOG_PREAMBLE, TEXT("Wintab Digitizer Services are not available."));
+			}
+			else
+			{
+				constexpr UINT OutputBufferSize = 1024;
+				TCHAR OutputBuffer[OutputBufferSize];
+				if (ensure(WtInfo(WTI_INTERFACE, IFC_WINTABID, OutputBuffer) <= OutputBufferSize))
+				{
+					InterfaceInfo.ID = OutputBuffer;
+
+					LogVerbose(LOG_PREAMBLE, FString::Format(TEXT("Wintab ID: {0}"), {OutputBuffer}));
+				}
+
+				WORD SpecVersion;
+				if (ensure(WtInfo(WTI_INTERFACE, IFC_SPECVERSION, &SpecVersion) == sizeof(SpecVersion)))
+				{
+					InterfaceInfo.SpecificationVersion = {static_cast<uint16>(SpecVersion >> 8), static_cast<uint16>(SpecVersion & 0xFF)};
+
+					LogVerbose(LOG_PREAMBLE, FString::Format(TEXT("Wintab Specification Version: {0}.{1}"),
+					                                         {InterfaceInfo.SpecificationVersion.Major, InterfaceInfo.SpecificationVersion.Minor}));
+				}
+
+				WORD ImplVersion;
+				if (ensure(WtInfo(WTI_INTERFACE, IFC_IMPLVERSION, &ImplVersion) == sizeof(ImplVersion)))
+				{
+					InterfaceInfo.ImplementationVersion = {static_cast<uint16>(ImplVersion >> 8), static_cast<uint16>(ImplVersion & 0xFF)};
+
+					LogVerbose(LOG_PREAMBLE, FString::Format(TEXT("Wintab Implementation Version: {0}.{1}"),
+					                                         {InterfaceInfo.ImplementationVersion.Major, InterfaceInfo.ImplementationVersion.Minor}));
+				}
+
+				/*UINT NumDevices;
+				if (ensure(WtInfo(WTI_INTERFACE, IFC_NDEVICES, &NumDevices) == sizeof(NumDevices)))
+				{
+					LogVerbose(LOG_PREAMBLE, FString::Format(TEXT("Wintab Number of Connected Devices: {0}"), {NumDevices}));
+				}*/
+
+				UINT MaxNumCursorTypes;
+				if (ensure(WtInfo(WTI_INTERFACE, IFC_NCURSORS, &MaxNumCursorTypes) == sizeof(MaxNumCursorTypes)))
+				{
+					InterfaceInfo.MaximumNumberOfCursorTypes = MaxNumCursorTypes;
+
+					LogVerbose(LOG_PREAMBLE, FString::Format(TEXT("Wintab Maximum Number of Cursor Types: {0}"), {MaxNumCursorTypes}));
+				}
+
+				UINT NumContexts;
+				if (ensure(WtInfo(WTI_INTERFACE, IFC_NCONTEXTS, &NumContexts) == sizeof(NumContexts)))
+				{
+					InterfaceInfo.MaximumNumberOfContexts = NumContexts;
+
+					LogVerbose(LOG_PREAMBLE, FString::Format(TEXT("Wintab Maximum Number of Contexts: {0}"), {NumContexts}));
+				}
+
+				UINT ContextOptions;
+				if (ensure(WtInfo(WTI_INTERFACE, IFC_CTXOPTIONS, &ContextOptions) == sizeof(ContextOptions)))
+				{
+					InterfaceInfo.bSupportsPacketMessages = ContextOptions & CXO_MESSAGES;
+					InterfaceInfo.bSupportsCursorChangeMessages = ContextOptions & CXO_CSRMESSAGES;
+
+					LogVerbose(LOG_PREAMBLE, FString::Format(TEXT("Wintab Supports System Cursor Contexts: {0}"), {ContextOptions & CXO_SYSTEM ? TEXT("Yes") : TEXT("No")}));
+					LogVerbose(LOG_PREAMBLE, FString::Format(TEXT("Wintab Supports Pen Windows Contexts: {0}"), {ContextOptions & CXO_PEN ? TEXT("Yes") : TEXT("No")}));
+					LogVerbose(LOG_PREAMBLE, FString::Format(TEXT("Wintab Supports WT_PACKET Messages: {0}"), {ContextOptions & CXO_MESSAGES ? TEXT("Yes") : TEXT("No")}));
+					LogVerbose(LOG_PREAMBLE, FString::Format(TEXT("Wintab Supports WT_CSRCHANGE  Messages: {0}"), {ContextOptions & CXO_CSRMESSAGES ? TEXT("Yes") : TEXT("No")}));
+					LogVerbose(LOG_PREAMBLE, FString::Format(TEXT("Wintab Supports Margins: {0}"), {ContextOptions & CXO_MARGIN ? TEXT("Yes") : TEXT("No")}));
+					LogVerbose(LOG_PREAMBLE, FString::Format(TEXT("Wintab Supports Inside Margins: {0}"), {ContextOptions & CXO_MGNINSIDE ? TEXT("Yes") : TEXT("No")}));
+				}
+			}
+		}
+	}
+
+	FWintabAPI::~FWintabAPI()
+	{
+		for (void* DllHandle : DllHandles)
+		{
+			FWindowsPlatformProcess::FreeDllHandle(DllHandle);
+		}
+	}
+
+	FName FWintabAPI::GetName()
+	{
+		static FName Name("Wintab");
+		return Name;
+	}
+
+	const FWintabAPI& FWintabAPI::GetInstance()
+	{
+		static FWintabAPI WintabAPI;
+		return WintabAPI;
+	}
+
+	bool FWintabAPI::SatisfiesRequirements() const
+	{
+		bool bSuccess = true;
+
+		bSuccess &= bWintabIsAvailable;
+		bSuccess &= GetWindowRect != nullptr;
+		bSuccess &= GetSystemMetrics != nullptr;
+		bSuccess &= WtInfo != nullptr;
+		bSuccess &= WtOpen != nullptr;
+		bSuccess &= WtClose != nullptr;
+		bSuccess &= WtPacket != nullptr;
+
+		return bSuccess;
+	}
+}
+
+#undef LOG_PREAMBLE
diff --git a/Engine/Plugins/Editor/StylusInput/Source/StylusInputWintab/Private/WintabAPI.h b/Engine/Plugins/Editor/StylusInput/Source/StylusInputWintab/Private/WintabAPI.h
new file mode 100644
--- /dev/null
+++ b/Engine/Plugins/Editor/StylusInput/Source/StylusInputWintab/Private/WintabAPI.h
@@ -0,0 +1,154 @@
+﻿// Copyright Epic Games, Inc. All Rights Reserved.
+
+#pragma once
+
+#include <Containers/Array.h>
+#include <Containers/UnrealString.h>
+#include <Templates/UniquePtr.h>
+
+#include <Windows/AllowWindowsPlatformTypes.h>
+
+	#define WIN32_LEAN_AND_MEAN
+	#include <windows.h>
+
+	#undef INT // WINTAB.H defines an INT macro that collides with a macro definition in AllowWindowsPlatformTypes.h when using Unity builds.
+		#include <WINTAB.H>
+		#define PACKETDATA (/*PK_CONTEXT |*/ PK_STATUS | PK_TIME | /*PK_CHANGED |*/ PK_SERIAL_NUMBER | PK_CURSOR | PK_BUTTONS | PK_X | PK_Y | PK_Z | PK_NORMAL_PRESSURE | PK_TANGENT_PRESSURE | PK_ORIENTATION | PK_ROTATION)
+		#define PACKETMODE PK_BUTTONS
+		#include <PKTDEF.H>
+
+#include <Windows/HideWindowsPlatformTypes.h>
+
+#define WTINFO_WORKAROUND_FOR_OUTPUT_BUFFER_SIZE_BUG 0
+
+namespace UE::StylusInput::Wintab
+{
+	struct FWintabInterfaceInfo
+	{
+		struct FVersion
+		{
+			uint16 Major;
+			uint16 Minor;
+		};
+
+		FString ID;
+		FVersion SpecificationVersion = { 0, 0 };
+		FVersion ImplementationVersion = { 0, 0 };
+		uint32 MaximumNumberOfCursorTypes = 0;
+		uint32 MaximumNumberOfContexts = 0;
+		bool bSupportsPacketMessages = false;
+		bool bSupportsCursorChangeMessages = false;
+	};
+
+	class FWintabAPI
+	{
+	public:
+		static FName GetName();
+		static const FWintabAPI& GetInstance();
+
+		bool SatisfiesRequirements() const;
+
+		using FFuncGetWindowRect = BOOL(HWND, LPRECT);
+		FFuncGetWindowRect* GetWindowRect = nullptr;
+
+		using FFuncGetSystemMetrics = int(int);
+		FFuncGetSystemMetrics* GetSystemMetrics = nullptr;
+
+		using FFuncWtInfo = UINT(UINT, UINT, LPVOID);
+		FFuncWtInfo* WtInfo = nullptr;
+
+		using FFuncWtOpen = HCTX(HWND, LPLOGCONTEXT, BOOL);
+		FFuncWtOpen* WtOpen = nullptr;
+
+		using FFuncWtClose = BOOL(HCTX);
+		FFuncWtClose* WtClose = nullptr;
+
+		using FFuncWtPacket = BOOL(HCTX, UINT, LPVOID);
+		FFuncWtPacket* WtPacket = nullptr;
+
+		FWintabInterfaceInfo InterfaceInfo;
+
+	private:
+
+		FWintabAPI();
+		~FWintabAPI();
+
+		TArray<void*> DllHandles;
+
+		bool bWintabIsAvailable = false;
+	};
+
+	inline uint32 HctxToUint32(const HCTX Hctx)
+	{
+		return IntCastChecked<uint32>(reinterpret_cast<uint64>(Hctx));
+	}
+
+	inline float Fix32ToFloat(const FIX32 Value)
+	{
+		return static_cast<float>(HIWORD(Value)) + LOWORD(Value) / 65536.0f;
+	}
+
+#if WTINFO_WORKAROUND_FOR_OUTPUT_BUFFER_SIZE_BUG
+
+	/*
+	 * Apparently, some drivers have issues with a nullptr being passed as output buffer for WtInfo() calls.
+	 * The workaround is to just always have a fixed size buffer and hope for the best.
+	 */
+
+	class FWintabInfoOutputBuffer
+	{
+	public:
+		void Allocate(const UINT Category, const UINT Index)
+		{
+			if (!Buffer.IsValid())
+			{
+				Buffer.Reset(new TCHAR[Size]);
+			}
+			Buffer[0] = 0;
+		}
+
+		TCHAR* Get() const { return Buffer.Get(); }
+		constexpr UINT SizeInBytes() { return Size * sizeof(TCHAR); }
+
+	private:
+		enum { Size = 1024 };
+		TUniquePtr<TCHAR[]> Buffer;
+	};
+
+#else
+
+	class FWintabInfoOutputBuffer
+	{
+	public:
+		FWintabInfoOutputBuffer()
+			: WintabAPI(FWintabAPI::GetInstance())
+		{
+		}
+
+		void Allocate(const UINT Category, const UINT Index)
+		{
+			// If the following call crashes, consider enabling the workaround available via WTINFO_WORKAROUND_FOR_OUTPUT_BUFFER_SIZE_BUG.
+			const UINT RequiredSize = WintabAPI.WtInfo(Category, Index, nullptr)  / sizeof(TCHAR);
+
+			if (Size < RequiredSize)
+			{
+				Size = RequiredSize;
+				Buffer.Reset(new TCHAR[Size]);
+			}
+			if (Size > 0)
+			{
+				Buffer[0] = 0;
+			}
+		}
+
+		TCHAR* Get() const { return Buffer.Get(); }
+		UINT SizeInBytes() const { return Size * sizeof(TCHAR); }
+
+	private:
+		const FWintabAPI& WintabAPI;
+		TUniquePtr<TCHAR[]> Buffer;
+		UINT Size = 0;
+	};
+
+#endif
+}
diff --git a/Engine/Plugins/Editor/StylusInput/Source/StylusInputWintab/Private/WintabInstance.cpp b/Engine/Plugins/Editor/StylusInput/Source/StylusInputWintab/Private/WintabInstance.cpp
new file mode 100644
--- /dev/null
+++ b/Engine/Plugins/Editor/StylusInput/Source/StylusInputWintab/Private/WintabInstance.cpp
@@ -0,0 +1,878 @@
+﻿// Copyright Epic Games, Inc. All Rights Reserved.
+
+#include "WintabInstance.h"
+
+#include <StylusInput.h>
+#include <StylusInputUtils.h>
+#include <Framework/Application/SlateApplication.h>
+
+#include "WintabAPI.h"
+
+#define LOCTEXT_NAMESPACE "WintabInstance"
+#define LOG_PREAMBLE "WintabInstance"
+
+using namespace UE::StylusInput::Private;
+
+namespace UE::StylusInput::Wintab
+{
+	bool SetupWintabContext(const FWintabAPI& WintabAPI, UINT WintabContextIndex, LOGCONTEXT& WintabContext)
+	{
+		bool bSuccess = true;
+
+		WintabContext.lcPktData = PACKETDATA;
+		WintabContext.lcPktMode = PACKETMODE;
+		WintabContext.lcOptions |= CXO_CSRMESSAGES | CXO_MESSAGES | CXO_SYSTEM;
+
+		AXIS AxisX = {};
+		AXIS AxisY = {};
+		AXIS AxisZ = {};
+
+		bSuccess &= WintabAPI.WtInfo(WTI_DEVICES + WintabContextIndex, DVC_X, &AxisX) == sizeof(AxisX);
+		bSuccess &= WintabAPI.WtInfo(WTI_DEVICES + WintabContextIndex, DVC_Y, &AxisY) == sizeof(AxisY);
+		bSuccess &= WintabAPI.WtInfo(WTI_DEVICES + WintabContextIndex, DVC_Z, &AxisZ) == sizeof(AxisZ);
+
+		if (!bSuccess)
+		{
+			LogError(LOG_PREAMBLE, FString::Format(TEXT("Failed to query device properties for Wintab context with index {0}."), {WintabContextIndex}));
+			return false;
+		}
+
+		WintabContext.lcOutExtX = AxisX.axMax - AxisX.axMin + 1;
+		WintabContext.lcOutExtY = AxisY.axMax - AxisY.axMin + 1;
+		WintabContext.lcOutExtZ = AxisZ.axMax - AxisZ.axMin + 1;
+
+		// In Wintab, the tablet origin is lower left. Move origin to upper left so that it corresponds to screen origin.
+		WintabContext.lcOutExtY = -WintabContext.lcOutExtY;
+
+		return bSuccess;
+	}
+
+	bool SetupTabletContextMetadata(const FWintabAPI& WintabAPI, UINT WintabContextIndex, HCTX WintabContextHandle,
+	                                const LOGCONTEXT& WintabContext, FTabletContext& TabletContext)
+	{
+		bool bSuccess = true;
+
+		TabletContext.WintabContextIndex = WintabContextIndex;
+		TabletContext.WintabContextHandle = WintabContextHandle;
+
+		TabletContext.InputRectangle = {
+			WintabContext.lcOutOrgX, WintabContext.lcOutOrgY,
+			WintabContext.lcOutOrgX + WintabContext.lcOutExtX, WintabContext.lcOutOrgY - WintabContext.lcOutExtY
+		};
+
+		FWintabInfoOutputBuffer OutputBuffer;
+
+		TabletContext.Name = WintabContext.lcName;
+		if (TabletContext.Name.IsEmpty())
+		{
+			OutputBuffer.Allocate(WTI_DEVICES + WintabContextIndex, DVC_NAME);
+			if (WintabAPI.WtInfo(WTI_DEVICES + WintabContextIndex, DVC_NAME, OutputBuffer.Get()) <= OutputBuffer.SizeInBytes())
+			{
+				TabletContext.Name = OutputBuffer.Get();
+			}
+			else
+			{
+				LogWarning(LOG_PREAMBLE, FString::Format(TEXT("Failed to query device name for Wintab context with index {0}."), {WintabContextIndex}));
+			}
+		}
+
+		OutputBuffer.Allocate(WTI_DEVICES + WintabContextIndex, DVC_PNPID);
+		if (WintabAPI.WtInfo(WTI_DEVICES + WintabContextIndex, DVC_PNPID, OutputBuffer.Get()) <= OutputBuffer.SizeInBytes())
+		{
+			TabletContext.PlugAndPlayID = OutputBuffer.Get();
+		}
+		else
+		{
+			TabletContext.PlugAndPlayID.Empty();
+			LogWarning(LOG_PREAMBLE, FString::Format(TEXT("Failed to query Plug and Play ID for Wintab context with index {0}."), {WintabContextIndex}));
+		}
+
+		UINT HardwareCapabilities;
+		if (WintabAPI.WtInfo(WTI_DEVICES + WintabContextIndex, DVC_HARDWARE, &HardwareCapabilities) != sizeof(HardwareCapabilities))
+		{
+			LogError(LOG_PREAMBLE, FString::Format(TEXT("Failed to query hardware capabilities for Wintab context with index {0}."), {WintabContextIndex}));
+			return false;
+		}
+
+		TabletContext.HardwareCapabilities =
+			(HardwareCapabilities & HWC_INTEGRATED ? ETabletHardwareCapabilities::Integrated : ETabletHardwareCapabilities::None)
+			| (HardwareCapabilities & HWC_TOUCH ? ETabletHardwareCapabilities::CursorMustTouch : ETabletHardwareCapabilities::None)
+			| (HardwareCapabilities & HWC_HARDPROX ? ETabletHardwareCapabilities::HardProximity : ETabletHardwareCapabilities::None)
+			| (HardwareCapabilities & HWC_PHYSID_CURSORS ? ETabletHardwareCapabilities::CursorsHavePhysicalIds : ETabletHardwareCapabilities::None);
+
+		/* CURSORS */
+
+		UINT FirstCursor, NumCursors;
+		if (WintabAPI.WtInfo(WTI_DEVICES + WintabContextIndex, DVC_FIRSTCSR, &FirstCursor) == sizeof(FirstCursor)
+			&& WintabAPI.WtInfo(WTI_DEVICES + WintabContextIndex, DVC_NCSRTYPES, &NumCursors) == sizeof(NumCursors))
+		{
+			TabletContext.WintabFirstCursor = FirstCursor;
+			TabletContext.WintabNumCursors = NumCursors;
+		}
+		else
+		{
+			LogError(LOG_PREAMBLE, FString::Format(TEXT("Failed to query cursor range for Wintab context with index {0}."), {WintabContextIndex}));
+			return false;
+		}
+
+		return bSuccess;
+	}
+
+	void SetPropertyStatus(FStylusInputPacket& Packet, const PACKET& WintabPacket, const int8* Data)
+	{
+		// Todo incorporate if tablet context is actually able to support Z values (or proximity?); if not, probably any packet is touching the tablet 
+		Packet.PenStatus = (WintabPacket.pkStatus & TPS_INVERT ? EPenStatus::CursorIsInverted : EPenStatus::None)
+			| (WintabPacket.pkZ <= 0 ? EPenStatus::CursorIsTouching : EPenStatus::None);
+	}
+	static_assert(std::is_same_v<FPacketPropertyHandler::FFuncSetProperty, decltype(SetPropertyStatus)>);
+
+	void SetPropertyTime(FStylusInputPacket& Packet, const PACKET& WintabPacket, const int8* Data)
+	{
+		Packet.TimerTick = WintabPacket.pkTime;
+	}
+	static_assert(std::is_same_v<FPacketPropertyHandler::FFuncSetProperty, decltype(SetPropertyTime)>);
+
+	void SetPropertySerialNumber(FStylusInputPacket& Packet, const PACKET& WintabPacket, const int8* Data)
+	{
+		Packet.SerialNumber = WintabPacket.pkSerialNumber;
+	}
+	static_assert(std::is_same_v<FPacketPropertyHandler::FFuncSetProperty, decltype(SetPropertySerialNumber)>);
+
+	void SetPropertyCursor(FStylusInputPacket& Packet, const PACKET& WintabPacket, const int8* Data)
+	{
+		const FWintabMessageHandler* MessageHandler = *reinterpret_cast<FWintabMessageHandler *const *>(Data);
+		Packet.CursorID = MessageHandler->GetCurrentStylusID();
+	}
+	static_assert(std::is_same_v<FPacketPropertyHandler::FFuncSetProperty, decltype(SetPropertyCursor)>);
+
+	void AssignSetPropertyCursorData(int8 *const Data, const FWintabMessageHandler* MessageHandler)
+	{
+		*reinterpret_cast<const FWintabMessageHandler**>(Data) = MessageHandler;
+	}
+
+	void SetPropertyButtons(FStylusInputPacket& Packet, const PACKET& WintabPacket, const int8* Data)
+	{
+		// Todo 
+		/*const WORD ButtonNumber = LOWORD(WintabPacket.pkButtons);
+		const WORD ButtonState = HIWORD(WintabPacket.pkButtons);*/
+	}
+	static_assert(std::is_same_v<FPacketPropertyHandler::FFuncSetProperty, decltype(SetPropertyButtons)>);
+
+	void SetPropertyX(FStylusInputPacket& Packet, const PACKET& WintabPacket, const int8* Data)
+	{
+		const int32& TabletContextInputWidth = reinterpret_cast<const int*>(Data)[0];
+		const int32& VirtualScreenWidth = reinterpret_cast<const int*>(Data)[1];
+
+		const RECT *const * WindowRect = reinterpret_cast<const RECT *const *>(Data + sizeof(int32) * 2);
+		const int32& WindowOffsetX = (*WindowRect)->left;
+
+		Packet.X = static_cast<float>(WintabPacket.pkX) * VirtualScreenWidth / TabletContextInputWidth - WindowOffsetX;
+	}
+	static_assert(std::is_same_v<FPacketPropertyHandler::FFuncSetProperty, decltype(SetPropertyX)>);
+
+	void AssignSetPropertyXData(int8 *const Data, const int32 TabletContextInputWidth, const int32 VirtualScreenWidth, const RECT* WindowRect)
+	{
+		int32* DataInt = reinterpret_cast<int32*>(Data);
+		*DataInt++ = TabletContextInputWidth;
+		*DataInt++ = VirtualScreenWidth;
+
+		const RECT** DataRect = reinterpret_cast<const RECT**>(DataInt);
+		*DataRect = WindowRect;
+	}
+
+	void SetPropertyY(FStylusInputPacket& Packet, const PACKET& WintabPacket, const int8* Data)
+	{
+		const int32& TabletContextInputHeight = reinterpret_cast<const int*>(Data)[0];
+		const int32& VirtualScreenHeight = reinterpret_cast<const int*>(Data)[1];
+
+		const RECT *const * WindowRect = reinterpret_cast<const RECT *const *>(Data + sizeof(int32) * 2);
+		const int32& WindowOffsetY = (*WindowRect)->top;
+
+		Packet.Y = static_cast<float>(WintabPacket.pkY) * VirtualScreenHeight / TabletContextInputHeight - WindowOffsetY;
+	}
+	static_assert(std::is_same_v<FPacketPropertyHandler::FFuncSetProperty, decltype(SetPropertyY)>);
+
+	void AssignSetPropertyYData(int8 *const Data, const int32 TabletContextInputHeight, const int32 VirtualScreenHeight, const RECT* WindowRect)
+	{
+		int32* DataInt = reinterpret_cast<int32*>(Data);
+		*DataInt++ = TabletContextInputHeight;
+		*DataInt++ = VirtualScreenHeight;
+
+		const RECT** DataRect = reinterpret_cast<const RECT**>(DataInt);
+		*DataRect = WindowRect;
+	}
+
+	void SetPropertyZ(FStylusInputPacket& Packet, const PACKET& WintabPacket, const int8* Data)
+	{
+		const float& InvResolution = *reinterpret_cast<const float*>(Data);
+
+		Packet.Z = static_cast<float>(WintabPacket.pkZ) * InvResolution;
+	}
+	static_assert(std::is_same_v<FPacketPropertyHandler::FFuncSetProperty, decltype(SetPropertyZ)>);
+
+	void AssignSetPropertyZData(int8 *const Data, const float Resolution)
+	{
+		*reinterpret_cast<float*>(Data) = 1.0f / Resolution;
+	}
+
+	void SetPropertyNormalPressure(FStylusInputPacket& Packet, const PACKET& WintabPacket, const int8* Data)
+	{
+		const float& InvExtent = *reinterpret_cast<const float*>(Data);
+
+		Packet.NormalPressure = static_cast<float>(WintabPacket.pkNormalPressure) * InvExtent;
+	}
+	static_assert(std::is_same_v<FPacketPropertyHandler::FFuncSetProperty, decltype(SetPropertyNormalPressure)>);
+
+	void AssignSetPropertyNormalPressureData(int8 *const Data, const int32 Extent)
+	{
+		*reinterpret_cast<float*>(Data) = 1.0f / Extent;
+	}
+
+	void SetPropertyTangentPressure(FStylusInputPacket& Packet, const PACKET& WintabPacket, const int8* Data)
+	{
+		const float& InvExtent = *reinterpret_cast<const float*>(Data);
+
+		Packet.TangentPressure = static_cast<float>(WintabPacket.pkTangentPressure) * InvExtent;
+	}
+	static_assert(std::is_same_v<FPacketPropertyHandler::FFuncSetProperty, decltype(SetPropertyTangentPressure)>);
+
+	void AssignSetPropertyTangentPressureData(int8 *const Data, const int32 Extent)
+	{
+		*reinterpret_cast<float*>(Data) = 1.0f / Extent;
+	}
+
+	void SetPropertyOrientation(FStylusInputPacket& Packet, const PACKET& WintabPacket, const int8* Data)
+	{
+		const float& InvScaleAzimuth = reinterpret_cast<const float*>(Data)[0];
+		const float& InvScaleAltitude = reinterpret_cast<const float*>(Data)[1];
+		const float& InvScaleTwist = reinterpret_cast<const float*>(Data)[2];
+
+		Packet.AzimuthOrientation = WintabPacket.pkOrientation.orAzimuth * InvScaleAzimuth;
+		Packet.AltitudeOrientation = WintabPacket.pkOrientation.orAltitude * InvScaleAltitude;
+		Packet.TwistOrientation = WintabPacket.pkOrientation.orTwist * InvScaleTwist;
+
+		const float AzimuthRadians = FMath::DegreesToRadians(Packet.AzimuthOrientation);
+		const float AltitudeRadians = FMath::DegreesToRadians(Packet.AltitudeOrientation);
+
+		const float X = FMath::Sin(AzimuthRadians) * FMath::Cos(AltitudeRadians);
+		const float Y = FMath::Cos(AzimuthRadians) * FMath::Cos(AltitudeRadians);
+
+		Packet.XTiltOrientation = FMath::RadiansToDegrees(FMath::Asin(-X));
+		Packet.YTiltOrientation = FMath::RadiansToDegrees(FMath::Asin(-Y));
+	}
+	static_assert(std::is_same_v<FPacketPropertyHandler::FFuncSetProperty, decltype(SetPropertyOrientation)>);
+
+	void AssignSetPropertyOrientationData(int8 *const Data, const float AzimuthResolution, const float AltitudeResolution, const float TwistResolution)
+	{
+		float* DataFloat = reinterpret_cast<float*>(Data);
+		*DataFloat++ = 1.0f / AzimuthResolution * 360.0f;
+		*DataFloat++ = 1.0f / AltitudeResolution * 360.0f;
+		*DataFloat = 1.0f / TwistResolution * 360.0f;
+	}
+
+	void SetPropertyRotation(FStylusInputPacket& Packet, const PACKET& WintabPacket, const int8* Data)
+	{
+		const float& InvScalePitch = reinterpret_cast<const float*>(Data)[0];
+		const float& InvScaleRoll = reinterpret_cast<const float*>(Data)[1];
+		const float& InvScaleYaw = reinterpret_cast<const float*>(Data)[2];
+
+		Packet.PitchRotation = WintabPacket.pkRotation.roPitch * InvScalePitch;
+		Packet.RollRotation = WintabPacket.pkRotation.roRoll * InvScaleRoll;
+		Packet.YawRotation = WintabPacket.pkRotation.roYaw * InvScaleYaw;
+	}
+	static_assert(std::is_same_v<FPacketPropertyHandler::FFuncSetProperty, decltype(SetPropertyRotation)>);
+
+	void AssignSetPropertyRotationData(int8 *const Data, const float PitchResolution, const float RollResolution, const float YawResolution)
+	{
+		float* DataFloat = reinterpret_cast<float*>(Data);
+		*DataFloat++ = 1.0f / PitchResolution * 360.0f;
+		*DataFloat++ = 1.0f / RollResolution * 360.0f;
+		*DataFloat = 1.0f / YawResolution * 360.0f;
+	}
+
+	bool SetupTabletContextPacketDescriptionData(const FWintabAPI& WintabAPI, UINT WintabContextIndex, const LOGCONTEXT& WintabContext,
+	                                             const RECT& WindowRect, FTabletContext& TabletContext, const FWintabMessageHandler& MessageHandler)
+	{
+		bool bSuccess = true;
+
+		TabletContext.SupportedProperties = ETabletSupportedProperties::None;
+		TabletContext.NumPacketProperties = 0;
+
+		WTPKT DataAvailableForAllCursors;
+		WTPKT DataAvailableForSomeCursors;
+		bSuccess &= WintabAPI.WtInfo(WTI_DEVICES + WintabContextIndex, DVC_PKTDATA, &DataAvailableForAllCursors) == sizeof(DataAvailableForAllCursors);
+		bSuccess &= WintabAPI.WtInfo(WTI_DEVICES + WintabContextIndex, DVC_CSRDATA, &DataAvailableForSomeCursors) == sizeof(DataAvailableForSomeCursors);
+
+		if (!bSuccess)
+		{
+			LogError(LOG_PREAMBLE, FString::Format(TEXT("Failed to query available packet data for Wintab context with index {0}."), {WintabContextIndex}));
+			return false;
+		}
+
+		const WTPKT DataAvailable = DataAvailableForAllCursors | DataAvailableForSomeCursors;
+
+		auto AddProperty = [&TabletContext](EPacketPropertyType Type, ETabletPropertyMetricUnit MetricUnit, int32 Minimum, int32 Maximum,
+		                                    float Resolution) -> FPacketPropertyHandler&
+		{
+			FPacketProperty& Property = TabletContext.PacketProperties[TabletContext.NumPacketProperties];
+			Property.Type = Type;
+			Property.MetricUnit = MetricUnit;
+			Property.Minimum = Minimum;
+			Property.Maximum = Maximum;
+			Property.Resolution = Resolution;
+
+			FPacketPropertyHandler& PropertyHandler = TabletContext.PacketPropertyHandlers[TabletContext.NumPacketProperties];
+
+			++TabletContext.NumPacketProperties;
+
+			return PropertyHandler;
+		};
+
+		if (DataAvailable & PK_STATUS)
+		{
+			TabletContext.SupportedProperties = ETabletSupportedProperties::PacketStatus;
+			FPacketPropertyHandler& PropertyHandler = AddProperty(EPacketPropertyType::Status, ETabletPropertyMetricUnit::Default, 0, 0, 0.0f);
+			PropertyHandler.SetProperty = &SetPropertyStatus;
+		}
+
+		if (DataAvailable & PK_TIME)
+		{
+			TabletContext.SupportedProperties = TabletContext.SupportedProperties | ETabletSupportedProperties::TimerTick;
+			FPacketPropertyHandler& PropertyHandler = AddProperty(EPacketPropertyType::Time, ETabletPropertyMetricUnit::Default, 0, 0, 0.0f);
+			PropertyHandler.SetProperty = &SetPropertyTime;
+		}
+
+		if (DataAvailable & PK_SERIAL_NUMBER)
+		{
+			TabletContext.SupportedProperties = TabletContext.SupportedProperties | ETabletSupportedProperties::SerialNumber;
+			FPacketPropertyHandler& PropertyHandler = AddProperty(EPacketPropertyType::SerialNumber, ETabletPropertyMetricUnit::Default, 0, 0, 0.0f);
+			PropertyHandler.SetProperty = &SetPropertySerialNumber;
+		}
+
+		if (DataAvailable & PK_CURSOR)
+		{
+			FPacketPropertyHandler& PropertyHandler = AddProperty(EPacketPropertyType::Cursor, ETabletPropertyMetricUnit::Default, 0, 0, 0.0f);
+			PropertyHandler.SetProperty = &SetPropertyCursor;
+			AssignSetPropertyCursorData(PropertyHandler.SetPropertyData, &MessageHandler);
+		}
+
+		if (DataAvailable & PK_BUTTONS)
+		{
+			FPacketPropertyHandler& PropertyHandler = AddProperty(EPacketPropertyType::Buttons, ETabletPropertyMetricUnit::Default, 0, 0, 0.0f);
+			PropertyHandler.SetProperty = &SetPropertyButtons;
+		}
+
+		if (DataAvailable & PK_X)
+		{
+			AXIS Axis;
+			if (WintabAPI.WtInfo(WTI_DEVICES + WintabContextIndex, DVC_X, &Axis) == sizeof(Axis))
+			{
+				TabletContext.SupportedProperties = TabletContext.SupportedProperties | ETabletSupportedProperties::X;
+				FPacketPropertyHandler& PropertyHandler = AddProperty(EPacketPropertyType::X, static_cast<ETabletPropertyMetricUnit>(Axis.axUnits), Axis.axMin,
+				                                                      Axis.axMax, Fix32ToFloat(Axis.axResolution));
+				PropertyHandler.SetProperty = &SetPropertyX;
+				AssignSetPropertyXData(PropertyHandler.SetPropertyData, WintabContext.lcInExtX, WintabAPI.GetSystemMetrics(SM_CXVIRTUALSCREEN), &WindowRect);
+			}
+			else
+			{
+				bSuccess = false;
+			}
+		}
+
+		if (DataAvailable & PK_Y)
+		{
+			AXIS Axis;
+			if (WintabAPI.WtInfo(WTI_DEVICES + WintabContextIndex, DVC_Y, &Axis) == sizeof(Axis))
+			{
+				TabletContext.SupportedProperties = TabletContext.SupportedProperties | ETabletSupportedProperties::Y;
+				FPacketPropertyHandler& PropertyHandler = AddProperty(EPacketPropertyType::Y, static_cast<ETabletPropertyMetricUnit>(Axis.axUnits), Axis.axMin,
+				                                                      Axis.axMax, Fix32ToFloat(Axis.axResolution));
+				PropertyHandler.SetProperty = &SetPropertyY;
+				AssignSetPropertyYData(PropertyHandler.SetPropertyData, WintabContext.lcInExtY, WintabAPI.GetSystemMetrics(SM_CYVIRTUALSCREEN), &WindowRect);
+			}
+			else
+			{
+				bSuccess = false;
+			}
+		}
+
+		if (DataAvailable & PK_Z)
+		{
+			AXIS Axis;
+			if (WintabAPI.WtInfo(WTI_DEVICES + WintabContextIndex, DVC_Z, &Axis) == sizeof(Axis))
+			{
+				TabletContext.SupportedProperties = TabletContext.SupportedProperties | ETabletSupportedProperties::Z;
+				FPacketPropertyHandler& PropertyHandler = AddProperty(EPacketPropertyType::Z, static_cast<ETabletPropertyMetricUnit>(Axis.axUnits), Axis.axMin,
+				                                                      Axis.axMax, Fix32ToFloat(Axis.axResolution));
+
+				const auto& Property = TabletContext.PacketProperties[TabletContext.NumPacketProperties - 1]; 
+				const float Resolution = [&Property]
+				{
+					switch (Property.MetricUnit)
+					{
+					case ETabletPropertyMetricUnit::Inches:
+						return Property.Resolution * 2.54f;
+					case ETabletPropertyMetricUnit::Centimeters:
+						return Property.Resolution;
+					default:
+						return 1.0f;
+					}
+				}();
+
+				PropertyHandler.SetProperty = &SetPropertyZ;
+				AssignSetPropertyZData(PropertyHandler.SetPropertyData, Resolution);
+			}
+			else
+			{
+				bSuccess = false;
+			}
+		}
+
+		if (DataAvailable & PK_NORMAL_PRESSURE)
+		{
+			AXIS Axis;
+			if (WintabAPI.WtInfo(WTI_DEVICES + WintabContextIndex, DVC_NPRESSURE, &Axis) == sizeof(Axis))
+			{
+				TabletContext.SupportedProperties = TabletContext.SupportedProperties | ETabletSupportedProperties::NormalPressure;
+				FPacketPropertyHandler& PropertyHandler = AddProperty(EPacketPropertyType::NormalPressure, static_cast<ETabletPropertyMetricUnit>(Axis.axUnits), Axis.axMin,
+				                                        Axis.axMax, Fix32ToFloat(Axis.axResolution));
+				PropertyHandler.SetProperty = &SetPropertyNormalPressure;
+				AssignSetPropertyNormalPressureData(PropertyHandler.SetPropertyData, Axis.axMax - Axis.axMin + 1);
+			}
+			else
+			{
+				bSuccess = false;
+			}
+		}
+
+		if (DataAvailable & PK_TANGENT_PRESSURE)
+		{
+			AXIS Axis;
+			if (WintabAPI.WtInfo(WTI_DEVICES + WintabContextIndex, DVC_TPRESSURE, &Axis) == sizeof(Axis))
+			{
+				TabletContext.SupportedProperties = TabletContext.SupportedProperties | ETabletSupportedProperties::TangentPressure;
+				FPacketPropertyHandler& PropertyHandler = AddProperty(EPacketPropertyType::TangentPressure, static_cast<ETabletPropertyMetricUnit>(Axis.axUnits), Axis.axMin,
+				                                        Axis.axMax, Fix32ToFloat(Axis.axResolution));
+				PropertyHandler.SetProperty = &SetPropertyTangentPressure;
+				AssignSetPropertyTangentPressureData(PropertyHandler.SetPropertyData, Axis.axMax - Axis.axMin + 1);
+			}
+			else
+			{
+				bSuccess = false;
+			}
+		}
+
+		if (DataAvailable & PK_ORIENTATION)
+		{
+			AXIS Orientation[3];
+			if (WintabAPI.WtInfo(WTI_DEVICES + WintabContextIndex, DVC_ORIENTATION, &Orientation) == sizeof(Orientation))
+			{
+				if (Orientation[0].axResolution && Orientation[1].axResolution)
+				{
+					TabletContext.SupportedProperties = TabletContext.SupportedProperties | ETabletSupportedProperties::AzimuthOrientation |
+						ETabletSupportedProperties::AltitudeOrientation | ETabletSupportedProperties::TwistOrientation |
+						ETabletSupportedProperties::XTiltOrientation | ETabletSupportedProperties::YTiltOrientation;
+
+					if (Orientation[0].axUnits != TU_CIRCLE || Orientation[1].axUnits != TU_CIRCLE || Orientation[2].axUnits != TU_CIRCLE)
+					{
+						LogWarning(LOG_PREAMBLE, FString::Format(
+							           TEXT("Units for orientation are not all TU_CIRCLE for Wintab context with index {0}."), {WintabContextIndex}));
+					}
+
+					FPacketPropertyHandler& PropertyHandler = AddProperty(EPacketPropertyType::Orientation, static_cast<ETabletPropertyMetricUnit>(Orientation[0].axUnits), 0,
+					                                        0, Fix32ToFloat(Orientation[0].axResolution));
+					PropertyHandler.SetProperty = &SetPropertyOrientation;
+					AssignSetPropertyOrientationData(PropertyHandler.SetPropertyData, Fix32ToFloat(Orientation[0].axResolution),
+					                                 Fix32ToFloat(Orientation[1].axResolution),
+					                                 Fix32ToFloat(Orientation[2].axResolution));
+				}
+			}
+			else
+			{
+				bSuccess = false;
+			}
+		}
+
+		if (DataAvailable & PK_ROTATION)
+		{
+			AXIS Rotation[3];
+			if (WintabAPI.WtInfo(WTI_DEVICES + WintabContextIndex, DVC_ROTATION, &Rotation) == sizeof(Rotation))
+			{
+				TabletContext.SupportedProperties = TabletContext.SupportedProperties | ETabletSupportedProperties::PitchRotation |
+					ETabletSupportedProperties::RollRotation | ETabletSupportedProperties::YawRotation;
+
+				if (Rotation[0].axUnits != TU_CIRCLE || Rotation[1].axUnits != TU_CIRCLE || Rotation[2].axUnits != TU_CIRCLE)
+				{
+					LogWarning(LOG_PREAMBLE, FString::Format(
+								 TEXT("Units for rotation are not all TU_CIRCLE for Wintab context with index {0}."), {WintabContextIndex}));
+				}
+
+				FPacketPropertyHandler& PropertyHandler = AddProperty(EPacketPropertyType::Rotation, static_cast<ETabletPropertyMetricUnit>(Rotation[0].axUnits), 0, 0,
+				                                        Fix32ToFloat(Rotation[0].axResolution));
+				PropertyHandler.SetProperty = &SetPropertyRotation;
+				AssignSetPropertyRotationData(PropertyHandler.SetPropertyData, Fix32ToFloat(Rotation[0].axResolution), Fix32ToFloat(Rotation[1].axResolution),
+				                              Fix32ToFloat(Rotation[2].axResolution));
+			}
+			else
+			{
+				bSuccess = false;
+			}
+		}
+
+		if (!bSuccess)
+		{
+			LogError(LOG_PREAMBLE, FString::Format(TEXT("Failed to query packet property information for Wintab context with index {0}."), {WintabContextIndex}));
+		}
+
+		return bSuccess;
+	}
+
+	bool SetupTabletContext(const FWintabAPI& WintabAPI, UINT WintabContextIndex, HCTX WintabContextHandle, const LOGCONTEXT& WintabContext,
+	                        const RECT& WindowRect, FTabletContext& TabletContext, const FWintabMessageHandler& MessageHandler)
+	{
+		bool bSuccess = true;
+
+		bSuccess &= SetupTabletContextMetadata(WintabAPI, WintabContextIndex, WintabContextHandle, WintabContext, TabletContext);
+		bSuccess &= SetupTabletContextPacketDescriptionData(WintabAPI, WintabContextIndex, WintabContext, WindowRect, TabletContext, MessageHandler);
+
+		return bSuccess;
+	}
+
+	FWintabInstance::FWintabInstance(const uint32 ID, const HWND OSWindowHandle)
+        : ID(ID)
+		, WintabAPI(FWintabAPI::GetInstance())
+		, OSWindowHandle(OSWindowHandle) 
+		, MessageHandler(this,
+		                 FGetTabletContextCallback::CreateRaw(this, &FWintabInstance::GetTabletContextInternal),
+		                 FGetStylusIDCallback::CreateRaw(this, &FWintabInstance::GetStylusID),
+		                 FUpdateWindowRectCallback::CreateRaw(this, &FWintabInstance::UpdateWindowRect),
+						 FUpdateTabletContextsCallback::CreateRaw(this, &FWintabInstance::UpdateTabletContexts))
+	{
+		if (FSlateApplication::IsInitialized())
+		{
+			if (const TSharedPtr<GenericApplication> WindowsApplication = FSlateApplication::Get().GetPlatformApplication())
+			{
+				static_cast<FWindowsApplication*>(WindowsApplication.Get())->AddMessageHandler(MessageHandler);
+			}
+		}
+
+		UpdateWindowRect();
+	    UpdateTabletContexts();
+    }
+
+	FWintabInstance::~FWintabInstance()
+	{
+		ClearTabletContexts();
+
+		if (FSlateApplication::IsInitialized())
+		{
+			if (const TSharedPtr<GenericApplication> WindowsApplication = FSlateApplication::Get().GetPlatformApplication())
+			{
+				static_cast<FWindowsApplication*>(WindowsApplication.Get())->RemoveMessageHandler(MessageHandler);
+			}
+		}
+
+		const FPacketStats& PacketStats = MessageHandler.GetPacketStats();
+
+		if (const uint32 NumInvalidPackets = PacketStats.GetNumInvalidPackets())
+		{
+			LogVerbose(LOG_PREAMBLE, FString::Format(
+						   TEXT("Wintab instance '{0}' encountered {1} invalid packets."), {FWintabInstance::GetName().ToString(), NumInvalidPackets}));
+		}
+
+		if (const uint32 NumLostPackets = PacketStats.GetNumLostPackets())
+		{
+			LogVerbose(LOG_PREAMBLE, FString::Format(
+						   TEXT("Wintab instance '{0}' encountered {1} lost packets."), {FWintabInstance::GetName().ToString(), NumLostPackets}));
+		}
+	}
+
+	bool FWintabInstance::AddEventHandler(IStylusInputEventHandler* EventHandler, const EEventHandlerThread Thread)
+	{
+        if (!EventHandler)
+        {
+            LogWarning(LOG_PREAMBLE, "Tried to add nullptr as event handler.");
+            return false;
+        }
+
+		if (Thread == EEventHandlerThread::OnGameThread)
+		{
+			return MessageHandler.AddEventHandler(EventHandler);
+		}
+
+		if (Thread == EEventHandlerThread::Asynchronous)
+		{
+			// TODO Handle asynchronous thread
+			LogError(LOG_PREAMBLE, "Asynchronous event handler is not supported yet.");
+			return false;
+		}
+
+        return false;
+	}
+
+	bool FWintabInstance::RemoveEventHandler(IStylusInputEventHandler* EventHandler)
+	{
+        if (!EventHandler)
+        {
+            LogWarning(LOG_PREAMBLE, "Tried to remove nullptr event handler.");
+            return false;
+        }
+
+		if (MessageHandler.RemoveEventHandler(EventHandler))
+		{
+			return true;
+		}
+
+		// TODO Handle asynchronous thread
+
+		return false;
+	}
+
+	const TSharedPtr<IStylusInputTabletContext> FWintabInstance::GetTabletContext(const uint32 TabletContextID)
+	{
+        return TabletContexts.Get(TabletContextID);
+    }
+
+	const TSharedPtr<IStylusInputStylusInfo> FWintabInstance::GetStylusInfo(const uint32 StylusID)
+	{
+		return StylusInfos.Get(StylusID);
+	}
+
+	float FWintabInstance::GetPacketsPerSecond(const EEventHandlerThread EventHandlerThread) const
+	{
+        return MessageHandler.GetPacketStats().GetPacketsPerSecond();
+    }
+
+	FName FWintabInstance::GetInterfaceName()
+	{
+		return FWintabAPI::GetName();
+	}
+
+	FText FWintabInstance::GetName()
+	{
+		return FText::Format(LOCTEXT("Wintab", "Wintab #{0}"), ID);
+	}
+
+	bool FWintabInstance::WasInitializedSuccessfully()
+	{
+		return true;
+	}
+
+	const FTabletContext* FWintabInstance::GetTabletContextInternal(const uint32 TabletContextID) const
+	{
+		return TabletContexts.Get(TabletContextID).Get();
+	}
+
+	uint32 FWintabInstance::GetStylusID(const uint32 TabletContextID, uint32 CursorIndex)
+	{
+		DWORD CursorPhysicalID;
+		UINT CursorType;
+		if (WintabAPI.WtInfo(WTI_CURSORS + CursorIndex, CSR_PHYSID, &CursorPhysicalID) == sizeof(CursorPhysicalID) &&
+			WintabAPI.WtInfo(WTI_CURSORS + CursorIndex, CSR_TYPE, &CursorType) == sizeof(CursorType))
+		{
+			const uint16 MaskedCursorId = MaskCursorId(CursorType);
+			const uint16 MaskedCursorType = MaskCursorType(CursorType);
+
+			const uint64 CursorID = static_cast<uint64>(MaskedCursorId) << 32 | CursorPhysicalID;
+			const TTuple<uint64, uint32>* Mapping = CursorIDToStylusIDMappings.FindByPredicate(
+				[CursorID](const TTuple<uint64, uint32>& Tuple)
+				{
+					return Tuple.Get<0>() == CursorID;
+				});
+
+			if (!Mapping)
+			{
+				uint32 StylusID = CursorPhysicalID;
+
+				// Resolve any collisions of physical IDs, which are only guaranteed to be unique within a (masked) cursor type.
+				while (StylusInfos.Contains(StylusID))
+				{
+					++StylusID;
+				}
+
+				Mapping = &CursorIDToStylusIDMappings.Emplace_GetRef(CursorID, StylusID);
+
+				const TSharedRef<FStylusInfo> StylusInfo = StylusInfos.Add(StylusID);
+
+				StylusInfo->WintabPhysicalID = CursorPhysicalID;
+				StylusInfo->WintabCursorType = MaskedCursorType;
+
+				const ECursorIndexType CursorIndexType = [&TabletContexts = TabletContexts, TabletContextID, CursorIndex]
+				{
+					if (const TSharedPtr<FTabletContext> TabletContextPtr = TabletContexts.Get(TabletContextID))
+					{
+						const FTabletContext& TabletContext = *TabletContextPtr;
+						if (TabletContext.WintabFirstCursor <= CursorIndex && CursorIndex < static_cast<uint8>(TabletContext.WintabFirstCursor + TabletContext.WintabNumCursors))
+						{
+							const int8 CursorIndexTypeInt = CursorIndex - TabletContext.WintabFirstCursor;
+							if (0 <= CursorIndexTypeInt && CursorIndexTypeInt < static_cast<int8>(ECursorIndexType::Num_Enumerators))
+							{
+								return static_cast<ECursorIndexType>(CursorIndexTypeInt);
+							}
+						}
+					}
+					return ECursorIndexType::Invalid_Enumerator;
+				}();
+
+				// If this is the inverted side of a pen, populate the stylus data with the non-inverted side instead.
+				if (CursorIndexType != ECursorIndexType::Invalid_Enumerator && CursorIsInverted(CursorIndexType))
+				{
+					--CursorIndex;
+				}
+
+				FWintabInfoOutputBuffer OutputBuffer;
+				OutputBuffer.Allocate(WTI_CURSORS + CursorIndex, CSR_NAME);
+				if (WintabAPI.WtInfo(WTI_CURSORS + CursorIndex, CSR_NAME, OutputBuffer.Get()) <= OutputBuffer.SizeInBytes())
+				{
+					StylusInfo->Name = FString::Format(TEXT("{0} ({1})"), {OutputBuffer.Get(), MaskedCursorTypeToString(MaskedCursorType)});
+				}
+				else
+				{
+					LogWarning(LOG_PREAMBLE, FString::Format(TEXT("Failed to query name for Wintab cursor with index {0}."), {CursorIndex}));
+				}
+
+				BYTE NumButtons;
+				if (WintabAPI.WtInfo(WTI_CURSORS + CursorIndex, CSR_BUTTONS, &NumButtons) == sizeof(NumButtons))
+				{
+					StylusInfo->Buttons.SetNum(NumButtons);
+
+					OutputBuffer.Allocate(WTI_CURSORS + CursorIndex, CSR_BTNNAMES);
+					if (WintabAPI.WtInfo(WTI_CURSORS + CursorIndex, CSR_BTNNAMES, OutputBuffer.Get()) <= OutputBuffer.SizeInBytes())
+					{
+						const TCHAR* Name = OutputBuffer.Get();
+						int32 ButtonIndex = 0;
+						while (*Name)
+						{
+							StylusInfo->Buttons[ButtonIndex].Name = Name;
+							Name += StylusInfo->Buttons[ButtonIndex].Name.Len() + 1;
+							++ButtonIndex;
+						}
+					}
+					else
+					{
+						LogWarning(LOG_PREAMBLE, FString::Format(TEXT("Failed to query button names for Wintab cursor with index {0}."), {CursorIndex}));
+					}
+
+					BYTE ButtonMap[32] = {};
+					if (WintabAPI.WtInfo(WTI_CURSORS + CursorIndex, CSR_BUTTONMAP, &ButtonMap) == sizeof(ButtonMap))
+					{
+						for (int32 ButtonIndex = 0; ButtonIndex < NumButtons; ++ButtonIndex)
+						{
+							// Todo
+						}
+					}
+					else
+					{
+						LogWarning(LOG_PREAMBLE, FString::Format(TEXT("Failed to query button map for Wintab cursor with index {0}."), {CursorIndex}));
+					}
+				}
+				else
+				{
+					LogWarning(LOG_PREAMBLE, FString::Format(TEXT("Failed to query number of buttons for Wintab cursor with index {0}."), {CursorIndex}));
+				}
+			}
+
+			// Make sure a StylusInfo with this ID exists, and the cursor type and physical ID match.
+			checkSlow(Mapping && StylusInfos.Contains(Mapping->Get<1>()) && StylusInfos.Get(Mapping->Get<1>()).IsValid()
+				&& StylusInfos.Get(Mapping->Get<1>())->WintabCursorType == MaskedCursorType
+				&& StylusInfos.Get(Mapping->Get<1>())->WintabPhysicalID == CursorPhysicalID);
+
+			return Mapping->Get<1>();
+		}
+		else
+		{
+			LogWarning(LOG_PREAMBLE, FString::Format(TEXT("Failed to query cursor metadata for cursor index {0}."), {CursorIndex}));
+		}
+
+		return 0;
+	}
+
+	void FWintabInstance::ClearTabletContexts()
+	{
+		for (int32 Index = 0, NumIndices = TabletContexts.Num(); Index < NumIndices; ++Index)
+		{
+			const HCTX WintabContextHandle = TabletContexts[Index]->WintabContextHandle;
+
+			if (WintabAPI.WtClose(WintabContextHandle))
+			{
+				LogVerbose(LOG_PREAMBLE, FString::Format(TEXT("Closed Wintab context with handle {0}."),
+														 {HctxToUint32(WintabContextHandle)}));
+			}
+			else
+			{
+				LogError(LOG_PREAMBLE, FString::Format(TEXT("Could not close Wintab context with handle {0}."),
+													   {HctxToUint32(WintabContextHandle)}));
+			}
+		}
+
+		TabletContexts.Clear();
+	}
+
+	void FWintabInstance::UpdateTabletContexts()
+	{
+		ClearTabletContexts();
+
+		UINT WintabContextIndex = 0;
+		LOGCONTEXTW WintabContext;
+
+		while (WintabAPI.WtInfo(WTI_DDCTXS + WintabContextIndex, 0, &WintabContext) == sizeof(WintabContext))
+		{
+			if (!SetupWintabContext(WintabAPI, WintabContextIndex, WintabContext))
+			{
+				LogError(LOG_PREAMBLE, FString::Format(TEXT("Failed to setup Wintab context with index {0}}."), {WintabContextIndex}));
+				continue;
+			}
+
+			const HCTX WintabContextHandle = WintabAPI.WtOpen(OSWindowHandle, &WintabContext, Windows::TRUE);
+			if (!WintabContextHandle)
+			{
+				LogError(LOG_PREAMBLE, FString::Format(TEXT("Failed to open Wintab context with index {0}."), {WintabContextIndex}));
+				continue;
+			}
+
+			const uint32 TabletContextID = HctxToUint32(WintabContextHandle);
+
+			TSharedRef<FTabletContext> TabletContextPtr = [&TabletContexts = TabletContexts, TabletContextID]
+			{
+				if (TabletContexts.Contains(TabletContextID))
+				{
+					return TabletContexts.Get(TabletContextID).ToSharedRef();
+				}
+				return TabletContexts.Add(TabletContextID);
+			}();
+
+			FTabletContext& TabletContext = *TabletContextPtr;
+
+			if (!SetupTabletContext(WintabAPI, WintabContextIndex, WintabContextHandle, WintabContext, WindowRect, TabletContext, MessageHandler))
+			{
+				LogError(LOG_PREAMBLE, FString::Format(TEXT("Failed setting up tablet context with ID {0}."), {TabletContextID}));
+
+				TabletContexts.Remove(TabletContextID);
+
+				if (!WintabAPI.WtClose(WintabContextHandle))
+				{
+					LogError(LOG_PREAMBLE, FString::Format(TEXT("Could not close Wintab context with handle {0}."), {HctxToUint32(WintabContextHandle)}));
+				}
+
+				continue;
+			}
+
+			LogVerbose(LOG_PREAMBLE, FString::Format(TEXT("Added tablet context for ID {0} [{1}, {2}]."), {
+				                                         TabletContext.ID, TabletContext.Name , TabletContext.PlugAndPlayID
+			                                         }));
+
+			++WintabContextIndex;
+		}
+	}
+
+	void FWintabInstance::UpdateWindowRect()
+	{
+		if (!ensure(WintabAPI.GetWindowRect(OSWindowHandle, &WindowRect)))
+		{
+			LogError(LOG_PREAMBLE, "Could not retrieve window rectangle; coordinates mapping will be incorrect!");
+		}
+	}
+}
+
+#undef LOG_PREAMBLE
+#undef LOCTEXT_NAMESPACE
diff --git a/Engine/Plugins/Editor/StylusInput/Source/StylusInputWintab/Private/WintabInstance.h b/Engine/Plugins/Editor/StylusInput/Source/StylusInputWintab/Private/WintabInstance.h
new file mode 100644
--- /dev/null
+++ b/Engine/Plugins/Editor/StylusInput/Source/StylusInputWintab/Private/WintabInstance.h
@@ -0,0 +1,60 @@
+﻿// Copyright Epic Games, Inc. All Rights Reserved.
+
+#pragma once
+
+#include <StylusInput.h>
+#include <Containers/Array.h>
+#include <Microsoft/MinimalWindowsApi.h>
+
+#include "WintabMessageHandler.h"
+#include "WintabTabletContext.h"
+#include "WintabStylus.h"
+
+namespace UE::StylusInput
+{
+	class IStylusInputEventHandler;
+}
+
+namespace UE::StylusInput::Wintab
+{
+	class FWintabAPI;
+
+	class FWintabInstance : public IStylusInputInstance
+	{
+	public:
+		explicit FWintabInstance(uint32 ID, HWND OSWindowHandle);
+		virtual ~FWintabInstance() override;
+
+		virtual bool AddEventHandler(IStylusInputEventHandler* EventHandler, EEventHandlerThread Thread) override;
+		virtual bool RemoveEventHandler(IStylusInputEventHandler* EventHandler) override;
+
+		virtual const TSharedPtr<IStylusInputTabletContext> GetTabletContext(uint32 TabletContextID) override;
+		virtual const TSharedPtr<IStylusInputStylusInfo> GetStylusInfo(uint32 StylusID) override;
+
+		virtual float GetPacketsPerSecond(EEventHandlerThread EventHandlerThread) const override;
+
+		virtual FName GetInterfaceName() override;
+		virtual FText GetName() override;
+
+		virtual bool WasInitializedSuccessfully() override;
+
+	private:
+		const FTabletContext* GetTabletContextInternal(uint32 TabletContextID) const;
+		uint32 GetStylusID(uint32 TabletContextID, uint32 CursorIndex);
+
+		void ClearTabletContexts();
+		void UpdateTabletContexts();
+		void UpdateWindowRect();
+
+		const uint32 ID;
+		const FWintabAPI& WintabAPI;
+		const HWND OSWindowHandle;
+		RECT WindowRect = {};
+
+		FTabletContextContainer TabletContexts;
+		FStylusInfoContainer StylusInfos;
+
+		FWintabMessageHandler MessageHandler;
+		TArray<TTuple<uint64, uint32>> CursorIDToStylusIDMappings;
+	};
+}
diff --git a/Engine/Plugins/Editor/StylusInput/Source/StylusInputWintab/Private/WintabInterface.cpp b/Engine/Plugins/Editor/StylusInput/Source/StylusInputWintab/Private/WintabInterface.cpp
new file mode 100644
--- /dev/null
+++ b/Engine/Plugins/Editor/StylusInput/Source/StylusInputWintab/Private/WintabInterface.cpp
@@ -0,0 +1,98 @@
+﻿// Copyright Epic Games, Inc. All Rights Reserved.
+
+#include "WintabInterface.h"
+
+#include <StylusInput.h>
+#include <StylusInputUtils.h>
+
+#include <GenericPlatform/GenericWindow.h>
+#include <Templates/SharedPointer.h>
+#include <Widgets/SWindow.h>
+
+#include "WintabInstance.h"
+#include "WintabAPI.h"
+
+#define LOCTEXT_NAMESPACE "WintabInterface"
+
+using namespace UE::StylusInput::Private;
+
+namespace UE::StylusInput::Wintab
+{
+	FName FWintabInterface::GetName()
+	{
+		return FWintabAPI::GetName();
+	}
+
+	IStylusInputInstance* FWintabInterface::CreateInstance(SWindow& Window)
+	{
+		if (FRefCountedInstance* ExistingRefCountedInstance = Instances.Find(&Window))
+		{
+			++ExistingRefCountedInstance->RefCount;
+			return ExistingRefCountedInstance->Instance.Get();
+		}
+
+		HWND OSWindowHandle = [&Window]
+		{
+			const TSharedPtr<const FGenericWindow> NativeWindow = Window.GetNativeWindow();
+			return NativeWindow.IsValid() ? static_cast<HWND>(NativeWindow->GetOSWindowHandle()) : nullptr;
+		}();
+
+		if (!OSWindowHandle)
+		{
+			LogError("WintabInterface", "Could not get native window handle.");
+			return nullptr;
+		}
+
+		FWintabInstance* NewInstance = Instances.Emplace(
+			&Window, {MakeUnique<FWintabInstance>(NextInstanceID++, OSWindowHandle), 1}).Instance.Get();
+		if (!ensureMsgf(NewInstance, TEXT("WintabInterface: Failed to create stylus input instance.")))
+		{
+			Instances.Remove(&Window);
+			return nullptr;
+		}
+
+		return NewInstance;
+	}
+
+	bool FWintabInterface::ReleaseInstance(IStylusInputInstance* Instance)
+	{
+		check(Instance);
+
+		FWintabInstance *const WindowsInstance = static_cast<FWintabInstance*>(Instance);
+
+		// Find existing instance
+		for (TTuple<SWindow*, FRefCountedInstance>& Entry : Instances)
+		{
+			FRefCountedInstance& RefCountedInstance = Entry.Get<1>();
+
+			if (RefCountedInstance.Instance.Get() == WindowsInstance)
+			{
+				// Decrease reference count
+				check(RefCountedInstance.RefCount > 0);
+				if (--RefCountedInstance.RefCount == 0)
+				{
+					// Delete if there are no references left
+					Instances.Remove(Entry.Key);
+				}
+
+				return true;
+			}
+		}
+
+		ensureMsgf(false, TEXT("WintabInterface: Failed to find provided instance."));
+		return false;
+	}
+
+	TUniquePtr<IStylusInputInterface> FWintabInterface::Create()
+	{
+		if (FWintabAPI::GetInstance().SatisfiesRequirements())
+		{
+			TUniquePtr<FWintabInterface> WintabInterface = MakeUnique<FWintabInterface>();
+			return TUniquePtr<IStylusInputInterface>(MoveTemp(WintabInterface));
+		}
+
+		return nullptr;
+	}
+}
+
+#undef LOCTEXT_NAMESPACE
diff --git a/Engine/Plugins/Editor/StylusInput/Source/StylusInputWintab/Private/WintabInterface.h b/Engine/Plugins/Editor/StylusInput/Source/StylusInputWintab/Private/WintabInterface.h
new file mode 100644
--- /dev/null
+++ b/Engine/Plugins/Editor/StylusInput/Source/StylusInputWintab/Private/WintabInterface.h
@@ -0,0 +1,36 @@
+﻿// Copyright Epic Games, Inc. All Rights Reserved.
+
+#pragma once
+
+#include <StylusInput.h>
+#include <StylusInputInterface.h>
+#include <Containers/Map.h>
+#include <Templates/UniquePtr.h>
+
+namespace UE::StylusInput::Wintab
+{
+	class FWintabInstance;
+
+	class FWintabInterface : public IStylusInputInterface
+	{
+	public:
+		virtual FName GetName() override;
+
+		virtual IStylusInputInstance* CreateInstance(SWindow& Window) override;
+		virtual bool ReleaseInstance(IStylusInputInstance* Instance) override;
+
+	private:
+		friend class FStylusInputWintabModule;
+
+		static TUniquePtr<IStylusInputInterface> Create();
+
+		struct FRefCountedInstance
+		{
+			TUniquePtr<FWintabInstance> Instance;
+			int32 RefCount;
+		};
+
+		TMap<SWindow*, FRefCountedInstance> Instances;
+		uint32 NextInstanceID = 0;
+	};
+}
diff --git a/Engine/Plugins/Editor/StylusInput/Source/StylusInputWintab/Private/WintabMessageHandler.cpp b/Engine/Plugins/Editor/StylusInput/Source/StylusInputWintab/Private/WintabMessageHandler.cpp
new file mode 100644
--- /dev/null
+++ b/Engine/Plugins/Editor/StylusInput/Source/StylusInputWintab/Private/WintabMessageHandler.cpp
@@ -0,0 +1,286 @@
+﻿// Copyright Epic Games, Inc. All Rights Reserved.
+
+#include "WintabMessageHandler.h"
+
+#include <StylusInputPacket.h>
+#include <StylusInputUtils.h>
+
+#include "WintabAPI.h"
+#include "WintabTabletContext.h"
+
+#define LOG_PREAMBLE "WintabMessageHandler"
+
+#define ENABLE_DEBUG_EVENTS_FOR_PACKET_MESSAGES 0
+#define ENABLE_DEBUG_EVENTS_FOR_CONTEXT_OPEN_CLOSE_UPDATE_MESSAGES 0
+#define ENABLE_DEBUG_EVENTS_FOR_CONTEXT_OVERLAP_MESSAGES 0
+#define ENABLE_DEBUG_EVENTS_FOR_PROXIMITY_CSRCHANGE_MESSAGES 0
+#define ENABLE_DEBUG_EVENTS_FOR_INFOCHANGE_MESSAGES 0
+#define ENABLE_DEBUG_EVENTS_FOR_PACKETEXT_MESSAGES 0
+#define ENABLE_DEBUG_EVENTS_FOR_INVALID_PACKETS 0
+
+using namespace UE::StylusInput::Private;
+
+namespace UE::StylusInput::Wintab
+{
+FWintabMessageHandler::FWintabMessageHandler(IStylusInputInstance* Instance, FGetTabletContextCallback&& GetTabletContextCallback,
+                                             FGetStylusIDCallback&& GetStylusIDCallback, FUpdateWindowRectCallback&& UpdateWindowRectCallback,
+                                             FUpdateTabletContextsCallback&& UpdateTabletContextsCallback)
+	: Instance(Instance)
+	, WintabAPI(FWintabAPI::GetInstance())
+	, GetTabletContextCallback(MoveTemp(GetTabletContextCallback))
+	, GetStylusIDCallback(MoveTemp(GetStylusIDCallback))
+	, UpdateWindowRectCallback(MoveTemp(UpdateWindowRectCallback))
+	, UpdateTabletContextsCallback(MoveTemp(UpdateTabletContextsCallback))
+{
+	check(this->GetTabletContextCallback.IsBound())
+	check(this->GetStylusIDCallback.IsBound())
+	check(this->UpdateWindowRectCallback.IsBound())
+	check(this->UpdateTabletContextsCallback.IsBound())
+}
+
+FWintabMessageHandler::~FWintabMessageHandler()
+{
+}
+
+bool FWintabMessageHandler::ProcessMessage(const HWND Hwnd, const uint32 Msg, const WPARAM WParam, const LPARAM LParam, int32& OutResult)
+{
+	bool bHandled = false;
+
+	switch (Msg)
+	{
+	case WT_PACKET:
+		{
+#if ENABLE_DEBUG_EVENTS_FOR_PACKET_MESSAGES
+			DebugEvent(FString::Format(TEXT("Received WT_PACKET(SerialNumber={0}, Context={1}) message."), {static_cast<uint32>(WParam), static_cast<uint32>(LParam)}));
+#endif
+
+			if (WParam != 0)
+			{
+				PACKET WintabPacket = {};
+				if (WintabAPI.WtPacket(reinterpret_cast<HCTX>(LParam), WParam, &WintabPacket))
+				{
+					if (bCursorChange)
+					{
+						ProcessCursorChange(static_cast<uint32>(LParam), &WintabPacket);
+					}
+
+					// Todo Process more than one packet? 
+
+					PacketStats.NewPacket(WintabPacket.pkSerialNumber);
+
+					const FStylusInputPacket Packet = ProcessPacket(static_cast<uint32>(LParam), &WintabPacket);
+
+					if (Packet.Type != EPacketType::Invalid)
+					{
+						for (IStylusInputEventHandler* EventHandler : EventHandlers)
+						{
+							EventHandler->OnPacket(Packet, Instance);
+						}
+					}
+				}
+				else
+				{
+					PacketStats.InvalidPacket();
+#if ENABLE_DEBUG_EVENTS_FOR_INVALID_PACKETS
+					DebugEvent(FString::Format(TEXT("Failed to retrieve Wintab packet number {0}."), {WParam}));
+#endif
+				}
+			}
+
+			bHandled = true;
+			break;
+		}
+	case WT_CTXOPEN:
+		{
+#if ENABLE_DEBUG_EVENTS_FOR_CONTEXT_OPEN_CLOSE_UPDATE_MESSAGES
+			DebugEvent(FString::Format(TEXT("Received WT_CTXOPEN message (Context={0}, Status={1})."), {WParam, LParam}));
+#endif
+
+			bHandled = true;
+			break;
+		}
+	case WT_CTXCLOSE:
+		{
+#if ENABLE_DEBUG_EVENTS_FOR_CONTEXT_OPEN_CLOSE_UPDATE_MESSAGES
+			DebugEvent(FString::Format(TEXT("Received WT_CTXCLOSE message (Context={0}, Status={1})."), {WParam, LParam}));
+#endif
+
+			bHandled = true;
+			break;
+		}
+	case WT_CTXUPDATE:
+		{
+#if ENABLE_DEBUG_EVENTS_FOR_CONTEXT_OPEN_CLOSE_UPDATE_MESSAGES
+			DebugEvent(FString::Format(TEXT("Received WT_CTXUPDATE message (Context={0}, Status={1})."), {WParam, LParam}));
+#endif
+
+			bHandled = true;
+			break;
+		}
+	case WT_CTXOVERLAP:
+		{
+#if ENABLE_DEBUG_EVENTS_FOR_CONTEXT_OVERLAP_MESSAGES
+			DebugEvent(FString::Format(TEXT("Received WT_CTXOVERLAP message (Context={0}, Status={1})."), {WParam, LParam}));
+#endif
+
+			bHandled = true;
+			break;
+		}
+	case WT_PROXIMITY:
+		{
+#if ENABLE_DEBUG_EVENTS_FOR_PROXIMITY_CSRCHANGE_MESSAGES
+			DebugEvent(FString::Format(TEXT("Received WT_PROXIMITY message (Context={0}, ContextEnter={1}, ProximityEnter={2})."), {
+											   static_cast<uint32>(WParam), LOWORD(LParam), HIWORD(LParam)
+										   }));
+#endif
+
+			if (LOWORD(LParam) || HIWORD(LParam))
+			{
+				// While the Wintab documentation is somewhat confusing about what is encoded in LParam, it appears that the high-order word is effectively
+				// indicating that the event is triggered by hardware proximity, and the low-order word is triggered both with and without hardware proximity
+				// support. Just in case, we trigger a cursor change if either of the flags are set.
+
+				bCursorChange = true;
+
+#if CHECK_MESSAGE_ORDER_ASSUMPTIONS
+				CursorChangeTabletContextID = static_cast<uint32>(WParam);
+#endif
+			}
+
+			bHandled = true;
+			break;
+		}
+	case WT_INFOCHANGE:
+		{
+#if ENABLE_DEBUG_EVENTS_FOR_INFOCHANGE_MESSAGES
+			DebugEvent(FString::Format(TEXT("Received WT_INFOCHANGE(Context={0}, Category={1}, Index={2}) message."), {WParam, LOWORD(LParam), HIWORD(LParam)}));
+#endif
+
+			UpdateTabletContextsCallback.Execute();
+
+			bHandled = true;
+			break;
+		}
+	case WT_CSRCHANGE:
+		{
+#if ENABLE_DEBUG_EVENTS_FOR_PROXIMITY_CSRCHANGE_MESSAGES
+			DebugEvent(FString::Format(TEXT("Received WT_CSRCHANGE message (SerialNumber={0}, Context={1})."), {static_cast<uint32>(WParam), static_cast<uint32>(LParam)}));
+#endif
+
+			bCursorChange = true;
+
+#if CHECK_MESSAGE_ORDER_ASSUMPTIONS
+			CursorChangeTabletContextID = static_cast<uint32>(LParam);
+			CursorChangeSerialNumber = static_cast<uint32>(WParam);
+#endif
+
+			bHandled = true;
+			break;
+		}
+	case WT_PACKETEXT:
+		{
+#if ENABLE_DEBUG_EVENTS_FOR_PACKETEXT_MESSAGES
+			DebugEvent("Received WT_PACKETEXT message.");
+#endif
+
+			bHandled = true;
+			break;
+		}
+	case WM_WINDOWPOSCHANGED:
+		{
+			UpdateWindowRectCallback.Execute();
+			break;
+		}
+	default:
+		;
+	}
+
+	return bHandled;
+}
+
+bool FWintabMessageHandler::AddEventHandler(IStylusInputEventHandler* EventHandler)
+{
+	check(EventHandler);
+
+	if (EventHandlers.Contains(EventHandler))
+	{
+		LogWarning(LOG_PREAMBLE, FString::Format(TEXT("Event handler '{0}' already exists."), {EventHandler->GetName()}));
+		return false;
+	}
+
+	EventHandlers.Add(EventHandler);
+
+	LogVerbose(LOG_PREAMBLE, FString::Format(TEXT("Event handler '{0}' was added."), {EventHandler->GetName()}));
+
+	return true;
+}
+
+bool FWintabMessageHandler::RemoveEventHandler(IStylusInputEventHandler* EventHandler)
+{
+	check(EventHandler);
+
+	const bool bWasRemoved = EventHandlers.Remove(EventHandler) > 0;
+
+	if (bWasRemoved)
+	{
+		LogVerbose(LOG_PREAMBLE, FString::Format(TEXT("Event handler '{0}' was removed."), {EventHandler->GetName()}));
+	}
+
+	return bWasRemoved;
+}
+
+void FWintabMessageHandler::DebugEvent(const FString& Message) const
+{
+	for (IStylusInputEventHandler* EventHandler : EventHandlers)
+	{
+		EventHandler->OnDebugEvent(Message, Instance);
+	}
+}
+
+void FWintabMessageHandler::ProcessCursorChange(const uint32 TabletContextID, const void *const WintabPacketPtr)
+{
+	const PACKET& WintabPacket = *static_cast<const PACKET*>(WintabPacketPtr);
+
+	check(bCursorChange);
+
+#if CHECK_MESSAGE_ORDER_ASSUMPTIONS
+	if (CursorChangeTabletContextID != TabletContextID || !(CursorChangeSerialNumber == 0 || CursorChangeSerialNumber == WintabPacket.pkSerialNumber))
+	{
+		DebugEvent("Cursor change assumptions violated: The next packet is not matching the previous WT_PROXIMITY/WT_CSRCHANGE messages.");
+	}
+#endif
+
+	CurrentStylusID = GetStylusIDCallback.Execute(TabletContextID, WintabPacket.pkCursor);
+
+	bCursorChange = false;
+
+#if CHECK_MESSAGE_ORDER_ASSUMPTIONS
+	CursorChangeTabletContextID = 0;
+	CursorChangeSerialNumber = 0;
+#endif
+}
+
+FStylusInputPacket FWintabMessageHandler::ProcessPacket(const uint32 TabletContextID, const void *const WintabPacketPtr)
+{
+	FStylusInputPacket Packet{TabletContextID};
+
+	if (const FTabletContext* TabletContext = GetTabletContextCallback.Execute(TabletContextID))
+	{
+		const PACKET& WintabPacket = *static_cast<const PACKET*>(WintabPacketPtr);
+
+		const bool bPacketOnDigitizer = !static_cast<bool>(TabletContext->SupportedProperties & ETabletSupportedProperties::Z) || WintabPacket.pkZ == 0;
+		Packet.Type = bPacketOnDigitizer
+						  ? bLastPacketOnDigitizer ? EPacketType::OnDigitizer : EPacketType::StylusDown
+						  : bLastPacketOnDigitizer ? EPacketType::StylusUp : EPacketType::AboveDigitizer;
+		bLastPacketOnDigitizer = bPacketOnDigitizer;
+
+		for (int32 Index = 0, Num = TabletContext->NumPacketProperties; Index < Num; ++Index)
+		{
+			const FPacketPropertyHandler& PropertyHandler = TabletContext->PacketPropertyHandlers[Index];
+			PropertyHandler.SetProperty(Packet, WintabPacket, PropertyHandler.SetPropertyData);
+		}
+	}
+
+	return Packet;
+}
+}
diff --git a/Engine/Plugins/Editor/StylusInput/Source/StylusInputWintab/Private/WintabMessageHandler.h b/Engine/Plugins/Editor/StylusInput/Source/StylusInputWintab/Private/WintabMessageHandler.h
new file mode 100644
--- /dev/null
+++ b/Engine/Plugins/Editor/StylusInput/Source/StylusInputWintab/Private/WintabMessageHandler.h
@@ -0,0 +1,66 @@
+﻿// Copyright Epic Games, Inc. All Rights Reserved.
+
+#pragma once
+
+#include <StylusInput.h>
+#include <Windows/WindowsApplication.h>
+
+#include "WintabStats.h"
+
+#define CHECK_MESSAGE_ORDER_ASSUMPTIONS NDEBUG
+
+namespace UE::StylusInput::Wintab
+{
+	struct FPacketProperty;
+	class FTabletContext;
+	class FStylusInfo;
+	struct FWindowContext;
+	class FWintabAPI;
+
+	DECLARE_DELEGATE_RetVal(const FWindowContext&, FGetWindowContextCallback);
+	DECLARE_DELEGATE_RetVal_OneParam(const FTabletContext*, FGetTabletContextCallback, uint32);
+	DECLARE_DELEGATE_RetVal_TwoParams(uint32, FGetStylusIDCallback, uint32, UINT);
+	DECLARE_DELEGATE(FUpdateWindowRectCallback);
+	DECLARE_DELEGATE(FUpdateTabletContextsCallback);
+
+	class FWintabMessageHandler : public IWindowsMessageHandler
+	{
+	public:
+		FWintabMessageHandler(IStylusInputInstance* Instance, FGetTabletContextCallback&& GetTabletContextCallback, FGetStylusIDCallback&& GetStylusIDCallback,
+		                      FUpdateWindowRectCallback&& UpdateWindowRectCallback, FUpdateTabletContextsCallback&& UpdateTabletContextsCallback);
+		virtual ~FWintabMessageHandler();
+
+		virtual bool ProcessMessage(HWND Hwnd, uint32 Msg, WPARAM WParam, LPARAM LParam, int32& OutResult) override;
+
+		bool AddEventHandler(IStylusInputEventHandler* EventHandler);
+		bool RemoveEventHandler(IStylusInputEventHandler* EventHandler);
+
+		uint32 GetCurrentStylusID() const { return CurrentStylusID; }
+		const FPacketStats& GetPacketStats() const { return PacketStats; }
+
+	private:
+		void DebugEvent(const FString& Message) const;
+
+		void ProcessCursorChange(uint32 TabletContextID, const void* WintabPacketPtr);
+		FStylusInputPacket ProcessPacket(uint32 TabletContextID, const void* WintabPacketPtr);
+
+		IStylusInputInstance *const Instance;
+		const FWintabAPI& WintabAPI;
+		FGetTabletContextCallback GetTabletContextCallback;
+		FGetStylusIDCallback GetStylusIDCallback;
+		FUpdateWindowRectCallback UpdateWindowRectCallback;
+		FUpdateTabletContextsCallback UpdateTabletContextsCallback;
+		TArray<IStylusInputEventHandler*> EventHandlers;
+		FPacketStats PacketStats;
+
+		uint32 CurrentStylusID = 0;
+
+		bool bCursorChange = false;
+		bool bLastPacketOnDigitizer = false;
+
+#if CHECK_MESSAGE_ORDER_ASSUMPTIONS
+		uint32 CursorChangeTabletContextID = 0;
+		uint32 CursorChangeSerialNumber = 0;
+#endif
+	};
+}
diff --git a/Engine/Plugins/Editor/StylusInput/Source/StylusInputWintab/Private/WintabStats.h b/Engine/Plugins/Editor/StylusInput/Source/StylusInputWintab/Private/WintabStats.h
new file mode 100644
--- /dev/null
+++ b/Engine/Plugins/Editor/StylusInput/Source/StylusInputWintab/Private/WintabStats.h
@@ -0,0 +1,61 @@
+﻿// Copyright Epic Games, Inc. All Rights Reserved.
+
+#pragma once
+
+#include <Windows/WindowsPlatformTime.h>
+
+namespace UE::StylusInput::Wintab
+{
+	class FPacketStats
+	{
+	public:
+		void NewPacket(uint32 SerialNumber)
+		{
+			const uint32 Timestamp = FPlatformTime::Cycles();
+
+			if (LatestTimestamp - EarliestTimestamp >= CyclesPerSecond)
+			{
+				PacketsPerSecond = EarliestTimestamp != LatestTimestamp ? NumPackets / ((LatestTimestamp - EarliestTimestamp) / CyclesPerSecond) : 0.0f;
+				EarliestTimestamp = Timestamp;
+				NumPackets = -1;
+			}
+			LatestTimestamp = Timestamp;
+
+			NumLostPackets += SerialNumber > LatestSerialNumber ? SerialNumber - LatestSerialNumber - 1 : 0;
+			LatestSerialNumber = SerialNumber;
+
+			++NumPackets;
+		}
+
+		void InvalidPacket()
+		{
+			++NumInvalidPackets;
+		}
+
+		float GetPacketsPerSecond() const
+		{
+			const uint32 Timestamp = FPlatformTime::Cycles();
+			return Timestamp - LatestTimestamp <= CyclesPerSecond ? PacketsPerSecond : 0.0f;
+		}
+
+		uint32 GetNumInvalidPackets() const
+		{
+			return NumInvalidPackets;
+		}
+
+		uint32 GetNumLostPackets() const
+		{
+			return NumLostPackets;
+		}
+
+	private:
+		uint32 EarliestTimestamp = 0;
+		uint32 LatestTimestamp = 0;
+		uint32 LatestSerialNumber = -1;
+		uint32 NumPackets = 0;
+		uint32 NumInvalidPackets = 0;
+		uint32 NumLostPackets = 0;
+		float PacketsPerSecond = 0.0f;
+		const double CyclesPerSecond = 1.0 / FPlatformTime::GetSecondsPerCycle();
+	};
+}
diff --git a/Engine/Plugins/Editor/StylusInput/Source/StylusInputWintab/Private/WintabStylus.h b/Engine/Plugins/Editor/StylusInput/Source/StylusInputWintab/Private/WintabStylus.h
new file mode 100644
--- /dev/null
+++ b/Engine/Plugins/Editor/StylusInput/Source/StylusInputWintab/Private/WintabStylus.h
@@ -0,0 +1,108 @@
+﻿// Copyright Epic Games, Inc. All Rights Reserved.
+
+#pragma once
+
+#include <StylusInputTabletContext.h>
+#include <StylusInputUtils.h>
+
+namespace UE::StylusInput::Wintab
+{
+	enum class ECursorType : uint16
+	{
+		GeneralStylus = 0x0802,
+		Airbrush = 0x0902,
+		ArtPen = 0x0804,
+		FourDMouse = 0x0004,
+		FiveButtonPuck = 0x0006,
+
+		CursorTypeMask = 0x0F06,	// Use this to mask any values coming from the Wintab API to get a valid cursor type.
+		CursorIdMask = 0xFF6		// Use this to mask any values coming from the Wintab API to differentiate between different devices, incl. colors. 
+	};
+
+	inline uint16 MaskCursorType(const UINT CursorType)
+	{
+		return static_cast<uint16>(CursorType) & static_cast<uint16>(ECursorType::CursorTypeMask);
+	}
+
+	inline uint16 MaskCursorId(const UINT CursorType)
+	{
+		return static_cast<uint16>(CursorType) & static_cast<uint16>(ECursorType::CursorIdMask);
+	}
+
+	inline FString MaskedCursorTypeToString(uint16 MaskedCursorType)
+	{
+		switch (static_cast<ECursorType>(MaskedCursorType))
+		{
+		case ECursorType::GeneralStylus:
+			return TEXT("General Stylus");
+		case ECursorType::Airbrush:
+			return TEXT("Airbrush");
+		case ECursorType::ArtPen:
+			return TEXT("Art Pen");
+		case ECursorType::FourDMouse:
+			return TEXT("4D Mouse");
+		case ECursorType::FiveButtonPuck:
+			return TEXT("5 Button Puck");
+		default:
+			return TEXT("Unknown Type");
+		}
+	}
+
+	enum class ECursorIndexType : int8
+	{
+		PuckDevice1 = 0,
+		StylusDevice1 = 1,
+		InvertedStylusDevice1 = 2,
+		PuckDevice2 = 3,
+		StylusDevice2 = 4,
+		InvertedStylusDevice2 = 5,
+
+		Num_Enumerators,
+		Invalid_Enumerator = INDEX_NONE
+	};
+
+	inline bool CursorIsInverted(ECursorIndexType Type)
+	{
+		return Type == ECursorIndexType::InvertedStylusDevice1 || Type == ECursorIndexType::InvertedStylusDevice2;
+	}
+
+	class FStylusButton final : public IStylusInputStylusButton
+	{
+	public:
+		virtual FString GetID() const override { return ID; }
+		virtual FString GetName() const override { return Name; }
+
+		FString ID;
+		FString Name;
+	};
+
+	class FStylusInfo final : public IStylusInputStylusInfo
+	{
+	public:
+		FStylusInfo() = default;
+
+		explicit FStylusInfo(const uint32 ID)
+			: ID(ID)
+		{
+		}
+
+		virtual uint32 GetID() const override { return ID; }
+		virtual FString GetName() const override { return Name; }
+		virtual uint32 GetNumButtons() const override { return Buttons.Num(); }
+
+		virtual const IStylusInputStylusButton* GetButton(int32 Index) const override
+		{
+			return 0 <= Index && Index < Buttons.Num() ? reinterpret_cast<const IStylusInputStylusButton*>(&Buttons[Index]) : nullptr;
+		}
+
+		uint32 ID = 0; // This will almost always be the physical ID, unless there is a rare collision between physical IDs of cursors of different types.
+		FString Name;
+		TArray<FStylusButton> Buttons;
+
+		uint32 WintabPhysicalID = -1;
+		uint16 WintabCursorType = -1;
+	};
+
+	using FStylusInfoContainer = Private::TSharedRefDataContainer<FStylusInfo, false>;
+	using FStylusInfoThreadSafeContainer = Private::TSharedRefDataContainer<FStylusInfo, false>;
+}
diff --git a/Engine/Plugins/Editor/StylusInput/Source/StylusInputWintab/Private/WintabTabletContext.h b/Engine/Plugins/Editor/StylusInput/Source/StylusInputWintab/Private/WintabTabletContext.h
new file mode 100644
--- /dev/null
+++ b/Engine/Plugins/Editor/StylusInput/Source/StylusInputWintab/Private/WintabTabletContext.h
@@ -0,0 +1,102 @@
+﻿// Copyright Epic Games, Inc. All Rights Reserved.
+
+#pragma once
+
+#include <StylusInputPacket.h>
+#include <StylusInputTabletContext.h>
+#include <StylusInputUtils.h>
+#include <Math/IntRect.h>
+
+#include "WintabAPI.h"
+
+namespace UE::StylusInput::Wintab
+{
+	enum class EPacketPropertyType : int8
+	{
+		Status = 0,
+		Time,
+		Changed,
+		SerialNumber,
+		Cursor,
+		Buttons,
+		X,
+		Y,
+		Z,
+		NormalPressure,
+		TangentPressure,
+		Orientation,
+		Rotation,
+
+		Num_Enumerators, // THIS IS NOT A VALID ENUMERATOR, BUT IT CAN BE USED TO QUERY HOW MANY PACKET PROPERTIES THERE ARE.
+		Invalid_Enumerator = INDEX_NONE // THIS IS NOT A VALID ENUMERATOR, BUT IT CAN BE USED TO INDICATE THAT A PACKET PROPERTY ENTRY IS INVALID.
+	};
+
+	enum class ETabletPropertyMetricUnit : int8
+	{
+		Default, // Units are unknown.
+		Inches,
+		Centimeters,
+		Circle,
+
+		Num_Enumerators, // THIS IS NOT A VALID ENUMERATOR, BUT IT CAN BE USED TO QUERY HOW MANY PACKET PROPERTIES THERE ARE.
+		Invalid_Enumerator = INDEX_NONE // THIS IS NOT A VALID ENUMERATOR, BUT IT CAN BE USED TO INDICATE THAT A PACKET PROPERTY ENTRY IS INVALID.
+	};
+
+	struct FPacketProperty
+	{
+		EPacketPropertyType Type = EPacketPropertyType::Invalid_Enumerator;
+		ETabletPropertyMetricUnit MetricUnit = ETabletPropertyMetricUnit::Default;
+		int32 Minimum = 0;
+		int32 Maximum = 0;
+		float Resolution = 0.0f;
+	};
+
+	struct FPacketPropertyHandler
+	{
+		using FFuncSetProperty = void(FStylusInputPacket&, const PACKET&, const int8*);
+		FFuncSetProperty* SetProperty = nullptr;
+
+		static constexpr int32 SetPropertyDataBufferLength = 16;
+		int8 SetPropertyData[SetPropertyDataBufferLength];
+	};
+
+	class FTabletContext final : public IStylusInputTabletContext
+	{
+	public:
+		explicit FTabletContext(const uint32 ID)
+			: ID(ID)
+		{}
+
+		virtual uint32 GetID() const override { return ID; }
+
+		virtual FString GetName() const override
+		{
+			if (!Name.IsEmpty() && !PlugAndPlayID.IsEmpty())
+			{
+				return FString::Format(TEXT("{0} ({1})"), {Name, PlugAndPlayID});
+			}
+			return !Name.IsEmpty() ? Name : PlugAndPlayID;
+		}
+
+		virtual FIntRect GetInputRectangle() const override { return InputRectangle; }
+		virtual ETabletHardwareCapabilities GetHardwareCapabilities() const override { return HardwareCapabilities; }
+		virtual ETabletSupportedProperties GetSupportedProperties() const override { return SupportedProperties; }
+
+		uint32 ID = 0;
+		FString Name;
+		FString PlugAndPlayID;
+		FIntRect InputRectangle;
+		ETabletSupportedProperties SupportedProperties = ETabletSupportedProperties::None;
+		ETabletHardwareCapabilities HardwareCapabilities = ETabletHardwareCapabilities::None;
+		uint32 NumPacketProperties = 0;
+		FPacketProperty PacketProperties[static_cast<uint32>(EPacketPropertyType::Num_Enumerators)] = {};
+		FPacketPropertyHandler PacketPropertyHandlers[static_cast<uint32>(EPacketPropertyType::Num_Enumerators)] = {};
+
+		UINT WintabContextIndex = -1;
+		HCTX WintabContextHandle = nullptr;
+		uint8 WintabFirstCursor = -1;
+		uint8 WintabNumCursors = -1;
+	};
+
+	using FTabletContextContainer = Private::TSharedRefDataContainer<FTabletContext, false>;
+}
diff --git a/Engine/Plugins/Editor/StylusInput/Source/StylusInputWintab/StylusInputWintab.Build.cs b/Engine/Plugins/Editor/StylusInput/Source/StylusInputWintab/StylusInputWintab.Build.cs
new file mode 100644
--- /dev/null
+++ b/Engine/Plugins/Editor/StylusInput/Source/StylusInputWintab/StylusInputWintab.Build.cs
@@ -0,0 +1,20 @@
+﻿// Copyright Epic Games, Inc. All Rights Reserved.
+
+namespace UnrealBuildTool.Rules
+{
+	public class StylusInputWintab : ModuleRules
+	{
+		public StylusInputWintab(ReadOnlyTargetRules Target) : base(Target)
+		{
+			PrivateDependencyModuleNames.AddRange(
+				[
+					"Core",
+					"Slate",
+					"SlateCore",
+					"StylusInput",
+					"Wintab"
+				]
+			);
+		}
+	}
+}
\ No newline at end of file
diff --git a/Engine/Plugins/Editor/StylusInput/StylusInput.uplugin b/Engine/Plugins/Editor/StylusInput/StylusInput.uplugin
--- a/Engine/Plugins/Editor/StylusInput/StylusInput.uplugin
+++ b/Engine/Plugins/Editor/StylusInput/StylusInput.uplugin
@@ -22,8 +22,21 @@
 			"LoadingPhase" : "Default"
 		},
 		{
-			"Name": "StylusInputDebugWidget",
-			"Type": "Editor"
+			"Name" : "StylusInputRealTimeStylus",
+			"Type" : "Editor",
+			"PlatformAllowList": ["Win64"],
+			"LoadingPhase" : "PostDefault"
+		},
+		{
+			"Name" : "StylusInputWintab",
+			"Type" : "Editor",
+			"PlatformAllowList": ["Win64"],
+			"LoadingPhase" : "PostDefault"
+		},
+		{
+			"Name" : "StylusInputDebugWidget",
+			"Type" : "Editor",
+			"LoadingPhase" : "PostEngineInit"
 		}
 	]
 }
