StylusInput: Fix issue with Wintab.dll not being available

diff --git a/Engine/Plugins/Editor/StylusInput/Source/StylusInputWintab/Private/WintabAPI.cpp b/Engine/Plugins/Editor/StylusInput/Source/StylusInputWintab/Private/WintabAPI.cpp
--- a/Engine/Plugins/Editor/StylusInput/Source/StylusInputWintab/Private/WintabAPI.cpp
+++ b/Engine/Plugins/Editor/StylusInput/Source/StylusInputWintab/Private/WintabAPI.cpp
@@ -42,16 +42,20 @@ namespace UE::StylusInput::Wintab
 		};
 
 		const TCHAR* User32DllName = TEXT("User32.dll");
-		void* User32DllHandle = GetDllHandle(User32DllName);
-		GetDllExport(User32DllName, User32DllHandle, TEXT("GetWindowRect"), GetWindowRect);
-		GetDllExport(User32DllName, User32DllHandle, TEXT("GetSystemMetrics"), GetSystemMetrics);
+		if (void* User32DllHandle = GetDllHandle(User32DllName))
+		{
+			GetDllExport(User32DllName, User32DllHandle, TEXT("GetWindowRect"), GetWindowRect);
+			GetDllExport(User32DllName, User32DllHandle, TEXT("GetSystemMetrics"), GetSystemMetrics);
+		}
 
 		const TCHAR* Wintab32DllName = TEXT("Wintab32.dll");
-		void* Wintab32DllHandle = GetDllHandle(Wintab32DllName);
-		GetDllExport(Wintab32DllName, Wintab32DllHandle, TEXT("WTInfoW"), WtInfo);
-		GetDllExport(Wintab32DllName, Wintab32DllHandle, TEXT("WTOpenW"), WtOpen);
-		GetDllExport(Wintab32DllName, Wintab32DllHandle, TEXT("WTClose"), WtClose);
-		GetDllExport(Wintab32DllName, Wintab32DllHandle, TEXT("WTPacket"), WtPacket);
+		if (void* Wintab32DllHandle = GetDllHandle(Wintab32DllName))
+		{
+			GetDllExport(Wintab32DllName, Wintab32DllHandle, TEXT("WTInfoW"), WtInfo);
+			GetDllExport(Wintab32DllName, Wintab32DllHandle, TEXT("WTOpenW"), WtOpen);
+			GetDllExport(Wintab32DllName, Wintab32DllHandle, TEXT("WTClose"), WtClose);
+			GetDllExport(Wintab32DllName, Wintab32DllHandle, TEXT("WTPacket"), WtPacket);
+		}
 
 		if (WtInfo)
 		{
