[Backout] - CL44879493
#rnx
Original CL Desc
-----------------------------------------------------------------
StylusInput: Fix issue with Wintab.dll not being available

diff --git a/Engine/Plugins/Editor/StylusInput/Source/StylusInput/Private/Windows/Wintab/WintabAPI.cpp b/Engine/Plugins/Editor/StylusInput/Source/StylusInput/Private/Windows/Wintab/WintabAPI.cpp
--- a/Engine/Plugins/Editor/StylusInput/Source/StylusInput/Private/Windows/Wintab/WintabAPI.cpp
+++ b/Engine/Plugins/Editor/StylusInput/Source/StylusInput/Private/Windows/Wintab/WintabAPI.cpp
@@ -42,20 +42,16 @@ namespace UE::StylusInput::Private::Wintab
 		};
 
 		const TCHAR* User32DllName = TEXT("User32.dll");
-		if (void* User32DllHandle = GetDllHandle(User32DllName))
-		{
-			GetDllExport(User32DllName, User32DllHandle, TEXT("GetWindowRect"), GetWindowRect);
-			GetDllExport(User32DllName, User32DllHandle, TEXT("GetSystemMetrics"), GetSystemMetrics);
-		}
+		void* User32DllHandle = GetDllHandle(User32DllName);
+		GetDllExport(User32DllName, User32DllHandle, TEXT("GetWindowRect"), GetWindowRect);
+		GetDllExport(User32DllName, User32DllHandle, TEXT("GetSystemMetrics"), GetSystemMetrics);
 
 		const TCHAR* Wintab32DllName = TEXT("Wintab32.dll");
-		if (void* Wintab32DllHandle = GetDllHandle(Wintab32DllName))
-		{
-			GetDllExport(Wintab32DllName, Wintab32DllHandle, TEXT("WTInfoW"), WtInfo);
-			GetDllExport(Wintab32DllName, Wintab32DllHandle, TEXT("WTOpenW"), WtOpen);
-			GetDllExport(Wintab32DllName, Wintab32DllHandle, TEXT("WTClose"), WtClose);
-			GetDllExport(Wintab32DllName, Wintab32DllHandle, TEXT("WTPacket"), WtPacket);
-		}
+		void* Wintab32DllHandle = GetDllHandle(Wintab32DllName);
+		GetDllExport(Wintab32DllName, Wintab32DllHandle, TEXT("WTInfoW"), WtInfo);
+		GetDllExport(Wintab32DllName, Wintab32DllHandle, TEXT("WTOpenW"), WtOpen);
+		GetDllExport(Wintab32DllName, Wintab32DllHandle, TEXT("WTClose"), WtClose);
+		GetDllExport(Wintab32DllName, Wintab32DllHandle, TEXT("WTPacket"), WtPacket);
 
 		if (WtInfo)
 		{
