StylusInput: Minor fixes to Wintab and RealTimeStylus interfaces

#rnx

diff --git a/Engine/Plugins/Editor/StylusInput/Source/StylusInputRealTimeStylus/Private/RealTimeStylusPluginAsync.h b/Engine/Plugins/Editor/StylusInput/Source/StylusInputRealTimeStylus/Private/RealTimeStylusPluginAsync.h
--- a/Engine/Plugins/Editor/StylusInput/Source/StylusInputRealTimeStylus/Private/RealTimeStylusPluginAsync.h
+++ b/Engine/Plugins/Editor/StylusInput/Source/StylusInputRealTimeStylus/Private/RealTimeStylusPluginAsync.h
@@ -12,7 +12,7 @@ namespace UE::StylusInput::RealTimeStylus
 		FRealTimeStylusPluginAsync(IStylusInputInstance* Instance, FGetWindowContextCallback&& GetWindowContext,
 		                           FUpdateTabletContextsCallback&& UpdateTabletContextsCallback, FUpdateStylusInfoCallback UpdateStylusInfoCallback,
 		                           IStylusInputEventHandler* EventHandler);
-		~FRealTimeStylusPluginAsync();
+		virtual ~FRealTimeStylusPluginAsync() override;
 
 		// IUnknown
 		virtual HRESULT QueryInterface(REFIID InterfaceID, void** InterfaceObject) override
diff --git a/Engine/Plugins/Editor/StylusInput/Source/StylusInputRealTimeStylus/Private/RealTimeStylusPluginBase.cpp b/Engine/Plugins/Editor/StylusInput/Source/StylusInputRealTimeStylus/Private/RealTimeStylusPluginBase.cpp
--- a/Engine/Plugins/Editor/StylusInput/Source/StylusInputRealTimeStylus/Private/RealTimeStylusPluginBase.cpp
+++ b/Engine/Plugins/Editor/StylusInput/Source/StylusInputRealTimeStylus/Private/RealTimeStylusPluginBase.cpp
@@ -47,7 +47,7 @@ namespace UE::StylusInput::RealTimeStylus
 			bSuccess = false;
 		}
 
-		IInkRectangle* Rectangle = nullptr;
+		TComPtr<IInkRectangle> Rectangle = nullptr;
 		if (Succeeded(InkTablet.get_MaximumInputRectangle(&Rectangle), LOG_PREAMBLE))
 		{
 			long Top, Left, Bottom, Right;
@@ -474,7 +474,7 @@ namespace UE::StylusInput::RealTimeStylus
 
 	bool SetupTabletContext(IRealTimeStylus* RealTimeStylus, const FWindowContext& WindowContext, FTabletContext& TabletContext)
 	{
-		IInkTablet* InkTablet = nullptr;
+		TComPtr<IInkTablet> InkTablet = nullptr;
 
 		if (Failed(RealTimeStylus->GetTabletFromTabletContextId(TabletContext.ID, &InkTablet), LOG_PREAMBLE))
 		{
diff --git a/Engine/Plugins/Editor/StylusInput/Source/StylusInputRealTimeStylus/Private/RealTimeStylusPluginBase.h b/Engine/Plugins/Editor/StylusInput/Source/StylusInputRealTimeStylus/Private/RealTimeStylusPluginBase.h
--- a/Engine/Plugins/Editor/StylusInput/Source/StylusInputRealTimeStylus/Private/RealTimeStylusPluginBase.h
+++ b/Engine/Plugins/Editor/StylusInput/Source/StylusInputRealTimeStylus/Private/RealTimeStylusPluginBase.h
@@ -39,7 +39,7 @@ namespace UE::StylusInput::RealTimeStylus
 	protected:
 		FRealTimeStylusPluginBase(IStylusInputInstance* Instance, FGetWindowContextCallback&& GetWindowContextCallback,
 		                          FUpdateTabletContextsCallback&& UpdateTabletContextsCallback, FUpdateStylusInfoCallback&& UpdateStylusInfoCallback);
-		~FRealTimeStylusPluginBase() = default;
+		virtual ~FRealTimeStylusPluginBase() = default;
 
 		virtual FString GetName() const = 0;
 
diff --git a/Engine/Plugins/Editor/StylusInput/Source/StylusInputRealTimeStylus/Private/RealTimeStylusPluginSync.h b/Engine/Plugins/Editor/StylusInput/Source/StylusInputRealTimeStylus/Private/RealTimeStylusPluginSync.h
--- a/Engine/Plugins/Editor/StylusInput/Source/StylusInputRealTimeStylus/Private/RealTimeStylusPluginSync.h
+++ b/Engine/Plugins/Editor/StylusInput/Source/StylusInputRealTimeStylus/Private/RealTimeStylusPluginSync.h
@@ -12,7 +12,7 @@ namespace UE::StylusInput::RealTimeStylus
 		FRealTimeStylusPluginSync(IStylusInputInstance* Instance, FGetWindowContextCallback&& GetWindowContext,
 		                          FUpdateTabletContextsCallback&& UpdateTabletContextsCallback, FUpdateStylusInfoCallback UpdateStylusInfoCallback,
 		                          IStylusInputEventHandler* EventHandler);
-		~FRealTimeStylusPluginSync();
+		virtual ~FRealTimeStylusPluginSync() override;
 
 		HRESULT CreateFreeThreadMarshaler();
 
diff --git a/Engine/Plugins/Editor/StylusInput/Source/StylusInputRealTimeStylus/Private/RealTimeStylusStats.h b/Engine/Plugins/Editor/StylusInput/Source/StylusInputRealTimeStylus/Private/RealTimeStylusStats.h
--- a/Engine/Plugins/Editor/StylusInput/Source/StylusInputRealTimeStylus/Private/RealTimeStylusStats.h
+++ b/Engine/Plugins/Editor/StylusInput/Source/StylusInputRealTimeStylus/Private/RealTimeStylusStats.h
@@ -11,30 +11,30 @@ namespace UE::StylusInput::RealTimeStylus
 	public:
 		void NewPacket()
 		{
-			const uint32 Timestamp = FPlatformTime::Cycles();
+			const uint64 Timestamp = FPlatformTime::Cycles64();
 
-			if (Latest - Earliest >= CyclesPerSecond)
+			if (LatestTimestamp - EarliestTimestamp >= CyclesPerSecond)
 			{
-				PacketsPerSecond = Earliest != Latest ? NumPackets / ((Latest - Earliest) / CyclesPerSecond) : 0.0f;
-				Earliest = Timestamp;
+				PacketsPerSecond = EarliestTimestamp != LatestTimestamp ? NumPackets / ((LatestTimestamp - EarliestTimestamp) / CyclesPerSecond) : 0.0f;
+				EarliestTimestamp = Timestamp;
 				NumPackets = -1;
 			}
 
-			Latest = Timestamp;
+			LatestTimestamp = Timestamp;
 			++NumPackets;
 		}
 
 		float GetPacketsPerSecond() const
 		{
-			const uint32 Timestamp = FPlatformTime::Cycles();
-			return Timestamp - Latest <= CyclesPerSecond ? PacketsPerSecond : 0.0f;
+			const uint64 Timestamp = FPlatformTime::Cycles64();
+			return Timestamp - LatestTimestamp <= CyclesPerSecond ? PacketsPerSecond : 0.0f;
 		}
 
 	private:
-		uint32 Earliest = 0;
-		uint32 Latest = 0;
+		uint64 EarliestTimestamp = 0;
+		uint64 LatestTimestamp = 0;
 		uint32 NumPackets = 0;
 		float PacketsPerSecond = 0.0f;
-		const double CyclesPerSecond = 1.0 / FPlatformTime::GetSecondsPerCycle();
+		const double CyclesPerSecond = 1.0 / FPlatformTime::GetSecondsPerCycle64();
 	};
 }
diff --git a/Engine/Plugins/Editor/StylusInput/Source/StylusInputWintab/Private/StylusInputWintabModule.cpp b/Engine/Plugins/Editor/StylusInput/Source/StylusInputWintab/Private/StylusInputWintabModule.cpp
--- a/Engine/Plugins/Editor/StylusInput/Source/StylusInputWintab/Private/StylusInputWintabModule.cpp
+++ b/Engine/Plugins/Editor/StylusInput/Source/StylusInputWintab/Private/StylusInputWintabModule.cpp
@@ -1,8 +1,9 @@
 ﻿// Copyright Epic Games, Inc. All Rights Reserved.
 
-#include <StylusInput.h>
+#include <StylusInputUtils.h>
 #include <Modules/ModuleManager.h>
 
+#include "WintabAPI.h"
 #include "WintabInterface.h"
 
 namespace UE::StylusInput::Wintab
@@ -19,10 +20,21 @@ namespace UE::StylusInput::Wintab
 
 	void FStylusInputWintabModule::StartupModule()
 	{
-		RegisteredInterface = FWintabInterface::Create();
-		if (RegisteredInterface)
+		if (FWintabAPI::DriverIsAvailable())
+		{
+			RegisteredInterface = FWintabInterface::Create();
+			if (RegisteredInterface)
+			{
+				RegisterInterface(RegisteredInterface.Get());
+			}
+			else
+			{
+				UE_LOG(Private::LogStylusInput, Warning, TEXT("Wintab interface initialization failed."));
+			}
+		}
+		else
 		{
-			RegisterInterface(RegisteredInterface.Get());
+			UE_LOG(Private::LogStylusInput, Log, TEXT("Wintab driver is not available; Wintab interface will be not be initialized."));
 		}
 	}
 
diff --git a/Engine/Plugins/Editor/StylusInput/Source/StylusInputWintab/Private/WintabAPI.cpp b/Engine/Plugins/Editor/StylusInput/Source/StylusInputWintab/Private/WintabAPI.cpp
--- a/Engine/Plugins/Editor/StylusInput/Source/StylusInputWintab/Private/WintabAPI.cpp
+++ b/Engine/Plugins/Editor/StylusInput/Source/StylusInputWintab/Private/WintabAPI.cpp
@@ -22,7 +22,7 @@ namespace UE::StylusInput::Wintab
 			}
 			else
 			{
-				LogWarning(LOG_PREAMBLE, FString::Format(TEXT("Could not get DLL handle for {0}."), {DllName}));
+				LogError(LOG_PREAMBLE, FString::Format(TEXT("Could not get DLL handle for {0}."), {DllName}));
 			}
 			return DllHandle;
 		};
@@ -140,6 +140,22 @@ namespace UE::StylusInput::Wintab
 		}
 	}
 
+	bool FWintabAPI::DriverIsAvailable()
+	{
+		if (void* Wintab32DllHandle = FWindowsPlatformProcess::GetDllHandle(TEXT("Wintab32.dll")))
+		{
+			if (FFuncWtInfo* WtInfo = static_cast<FFuncWtInfo*>(FWindowsPlatformProcess::GetDllExport(Wintab32DllHandle, TEXT("WTInfoW"))))
+			{
+				if (WtInfo(0, 0, nullptr) != 0)
+				{
+					return true;
+				}
+			}
+		}
+
+		return false;
+	}
+
 	FName FWintabAPI::GetName()
 	{
 		static FName Name("Wintab");
diff --git a/Engine/Plugins/Editor/StylusInput/Source/StylusInputWintab/Private/WintabAPI.h b/Engine/Plugins/Editor/StylusInput/Source/StylusInputWintab/Private/WintabAPI.h
--- a/Engine/Plugins/Editor/StylusInput/Source/StylusInputWintab/Private/WintabAPI.h
+++ b/Engine/Plugins/Editor/StylusInput/Source/StylusInputWintab/Private/WintabAPI.h
@@ -19,7 +19,7 @@
 
 #include <Windows/HideWindowsPlatformTypes.h>
 
-#define WTINFO_WORKAROUND_FOR_OUTPUT_BUFFER_SIZE_BUG 0
+#define WTINFO_WORKAROUND_FOR_OUTPUT_BUFFER_SIZE_BUG 1
 
 namespace UE::StylusInput::Wintab
 {
@@ -43,6 +43,8 @@ namespace UE::StylusInput::Wintab
 	class FWintabAPI
 	{
 	public:
+		static bool DriverIsAvailable();
+
 		static FName GetName();
 		static const FWintabAPI& GetInstance();
 
@@ -85,7 +87,7 @@ namespace UE::StylusInput::Wintab
 
 	inline float Fix32ToFloat(const FIX32 Value)
 	{
-		return static_cast<float>(HIWORD(Value)) + LOWORD(Value) / 65536.0f;
+		return Value * (1.0f / 65536.0f);
 	}
 
 #if WTINFO_WORKAROUND_FOR_OUTPUT_BUFFER_SIZE_BUG
diff --git a/Engine/Plugins/Editor/StylusInput/Source/StylusInputWintab/Private/WintabStats.h b/Engine/Plugins/Editor/StylusInput/Source/StylusInputWintab/Private/WintabStats.h
--- a/Engine/Plugins/Editor/StylusInput/Source/StylusInputWintab/Private/WintabStats.h
+++ b/Engine/Plugins/Editor/StylusInput/Source/StylusInputWintab/Private/WintabStats.h
@@ -11,7 +11,7 @@ namespace UE::StylusInput::Wintab
 	public:
 		void NewPacket(uint32 SerialNumber)
 		{
-			const uint32 Timestamp = FPlatformTime::Cycles();
+			const uint64 Timestamp = FPlatformTime::Cycles64();
 
 			if (LatestTimestamp - EarliestTimestamp >= CyclesPerSecond)
 			{
@@ -34,7 +34,7 @@ namespace UE::StylusInput::Wintab
 
 		float GetPacketsPerSecond() const
 		{
-			const uint32 Timestamp = FPlatformTime::Cycles();
+			const uint64 Timestamp = FPlatformTime::Cycles64();
 			return Timestamp - LatestTimestamp <= CyclesPerSecond ? PacketsPerSecond : 0.0f;
 		}
 
@@ -49,13 +49,13 @@ namespace UE::StylusInput::Wintab
 		}
 
 	private:
-		uint32 EarliestTimestamp = 0;
-		uint32 LatestTimestamp = 0;
+		uint64 EarliestTimestamp = 0;
+		uint64 LatestTimestamp = 0;
 		uint32 LatestSerialNumber = -1;
 		uint32 NumPackets = 0;
 		uint32 NumInvalidPackets = 0;
 		uint32 NumLostPackets = 0;
 		float PacketsPerSecond = 0.0f;
-		const double CyclesPerSecond = 1.0 / FPlatformTime::GetSecondsPerCycle();
+		const double CyclesPerSecond = 1.0 / FPlatformTime::GetSecondsPerCycle64();
 	};
 }
