StylusInput: Make Wintab API calls more robust/forgiving

diff --git a/Engine/Plugins/Editor/StylusInput/Source/StylusInput/Public/StylusInput.h b/Engine/Plugins/Editor/StylusInput/Source/StylusInput/Public/StylusInput.h
--- a/Engine/Plugins/Editor/StylusInput/Source/StylusInput/Public/StylusInput.h
+++ b/Engine/Plugins/Editor/StylusInput/Source/StylusInput/Public/StylusInput.h
@@ -133,7 +133,7 @@ namespace UE::StylusInput
 		 * 
 		 * @returns true if this instance was initialized successfully and able to provide stylus input events.
 		 */
-		virtual bool WasInitializedSuccessfully() = 0;
+		virtual bool WasInitializedSuccessfully() const = 0;
 	};
 
 	/**
diff --git a/Engine/Plugins/Editor/StylusInput/Source/StylusInput/Public/StylusInputUtils.h b/Engine/Plugins/Editor/StylusInput/Source/StylusInput/Public/StylusInputUtils.h
--- a/Engine/Plugins/Editor/StylusInput/Source/StylusInput/Public/StylusInputUtils.h
+++ b/Engine/Plugins/Editor/StylusInput/Source/StylusInput/Public/StylusInputUtils.h
@@ -20,6 +20,11 @@ namespace UE::StylusInput::Private
 		UE_LOG(LogStylusInput, Warning,  TEXT("%s: %s"), *Preamble, *Message);
 	}
 
+	inline void Log(const FString& Preamble, const FString& Message)
+	{
+		UE_LOG(LogStylusInput, Log,  TEXT("%s: %s"), *Preamble, *Message);
+	}
+
 	inline void LogVerbose(const FString& Preamble, const FString& Message)
 	{
 		UE_LOG(LogStylusInput, Verbose,  TEXT("%s: %s"), *Preamble, *Message);
diff --git a/Engine/Plugins/Editor/StylusInput/Source/StylusInputRealTimeStylus/Private/RealTimeStylusAPI.cpp b/Engine/Plugins/Editor/StylusInput/Source/StylusInputRealTimeStylus/Private/RealTimeStylusAPI.cpp
--- a/Engine/Plugins/Editor/StylusInput/Source/StylusInputRealTimeStylus/Private/RealTimeStylusAPI.cpp
+++ b/Engine/Plugins/Editor/StylusInput/Source/StylusInputRealTimeStylus/Private/RealTimeStylusAPI.cpp
@@ -149,7 +149,7 @@ namespace UE::StylusInput::RealTimeStylus
 		return API;
 	}
 
-	bool FRealTimeStylusAPI::SatisfiesRequirements() const
+	bool FRealTimeStylusAPI::IsValid() const
 	{
 		bool bSuccess = true;
 
diff --git a/Engine/Plugins/Editor/StylusInput/Source/StylusInputRealTimeStylus/Private/RealTimeStylusAPI.h b/Engine/Plugins/Editor/StylusInput/Source/StylusInputRealTimeStylus/Private/RealTimeStylusAPI.h
--- a/Engine/Plugins/Editor/StylusInput/Source/StylusInputRealTimeStylus/Private/RealTimeStylusAPI.h
+++ b/Engine/Plugins/Editor/StylusInput/Source/StylusInputRealTimeStylus/Private/RealTimeStylusAPI.h
@@ -18,7 +18,7 @@ namespace UE::StylusInput::RealTimeStylus
 		static FName GetName();
 		static const FRealTimeStylusAPI& GetInstance();
 
-		bool SatisfiesRequirements() const;
+		bool IsValid() const;
 
 		using FFuncGetTickCount = DWORD();
 		FFuncGetTickCount* GetTickCount = nullptr;
diff --git a/Engine/Plugins/Editor/StylusInput/Source/StylusInputRealTimeStylus/Private/RealTimeStylusInstance.cpp b/Engine/Plugins/Editor/StylusInput/Source/StylusInputRealTimeStylus/Private/RealTimeStylusInstance.cpp
--- a/Engine/Plugins/Editor/StylusInput/Source/StylusInputRealTimeStylus/Private/RealTimeStylusInstance.cpp
+++ b/Engine/Plugins/Editor/StylusInput/Source/StylusInputRealTimeStylus/Private/RealTimeStylusInstance.cpp
@@ -215,7 +215,7 @@ namespace UE::StylusInput::RealTimeStylus
 		return FText::Format(LOCTEXT("RealTimeStylus", "RealTimeStylus #{0}"), ID);
 	}
 
-	bool FRealTimeStylusInstance::WasInitializedSuccessfully()
+	bool FRealTimeStylusInstance::WasInitializedSuccessfully() const
 	{
 		return bWasInitializedSuccessfully;
 	}
diff --git a/Engine/Plugins/Editor/StylusInput/Source/StylusInputRealTimeStylus/Private/RealTimeStylusInstance.h b/Engine/Plugins/Editor/StylusInput/Source/StylusInputRealTimeStylus/Private/RealTimeStylusInstance.h
--- a/Engine/Plugins/Editor/StylusInput/Source/StylusInputRealTimeStylus/Private/RealTimeStylusInstance.h
+++ b/Engine/Plugins/Editor/StylusInput/Source/StylusInputRealTimeStylus/Private/RealTimeStylusInstance.h
@@ -40,7 +40,7 @@ namespace UE::StylusInput::RealTimeStylus
 		virtual FName GetInterfaceName() override;
 		virtual FText GetName() override;
 
-		virtual bool WasInitializedSuccessfully() override;
+		virtual bool WasInitializedSuccessfully() const override;
 
 	private:
 		void EnablePlugin(EEventHandlerThread EventHandlerThread, IStylusInputEventHandler* EventHandler);
diff --git a/Engine/Plugins/Editor/StylusInput/Source/StylusInputRealTimeStylus/Private/RealTimeStylusInterface.cpp b/Engine/Plugins/Editor/StylusInput/Source/StylusInputRealTimeStylus/Private/RealTimeStylusInterface.cpp
--- a/Engine/Plugins/Editor/StylusInput/Source/StylusInputRealTimeStylus/Private/RealTimeStylusInterface.cpp
+++ b/Engine/Plugins/Editor/StylusInput/Source/StylusInputRealTimeStylus/Private/RealTimeStylusInterface.cpp
@@ -87,7 +87,7 @@ namespace UE::StylusInput::RealTimeStylus
 
 	TUniquePtr<IStylusInputInterface> FRealTimeStylusInterface::Create()
 	{
-		if (FRealTimeStylusAPI::GetInstance().SatisfiesRequirements())
+		if (FRealTimeStylusAPI::GetInstance().IsValid())
 		{
 			TUniquePtr<FRealTimeStylusInterface> RealTimeStylusInterface = MakeUnique<FRealTimeStylusInterface>();
 			return TUniquePtr<IStylusInputInterface>(MoveTemp(RealTimeStylusInterface));
diff --git a/Engine/Plugins/Editor/StylusInput/Source/StylusInputWintab/Private/WintabAPI.cpp b/Engine/Plugins/Editor/StylusInput/Source/StylusInputWintab/Private/WintabAPI.cpp
--- a/Engine/Plugins/Editor/StylusInput/Source/StylusInputWintab/Private/WintabAPI.cpp
+++ b/Engine/Plugins/Editor/StylusInput/Source/StylusInputWintab/Private/WintabAPI.cpp
@@ -57,78 +57,9 @@ namespace UE::StylusInput::Wintab
 			GetDllExport(Wintab32DllName, Wintab32DllHandle, TEXT("WTPacket"), WtPacket);
 		}
 
-		if (WtInfo)
+		if (IsValid())
 		{
-			bWintabIsAvailable = WtInfo(0, 0, nullptr) != 0;
-			if (!bWintabIsAvailable)
-			{
-				LogVerbose(LOG_PREAMBLE, TEXT("Wintab Digitizer Services are not available."));
-			}
-			else
-			{
-				constexpr UINT OutputBufferSize = 1024;
-				TCHAR OutputBuffer[OutputBufferSize];
-				if (ensure(WtInfo(WTI_INTERFACE, IFC_WINTABID, OutputBuffer) <= OutputBufferSize * sizeof(TCHAR)))
-				{
-					InterfaceInfo.ID = OutputBuffer;
-
-					LogVerbose(LOG_PREAMBLE, FString::Format(TEXT("Wintab ID: {0}"), {OutputBuffer}));
-				}
-
-				WORD SpecVersion;
-				if (ensure(WtInfo(WTI_INTERFACE, IFC_SPECVERSION, &SpecVersion) == sizeof(SpecVersion)))
-				{
-					InterfaceInfo.SpecificationVersion = {static_cast<uint16>(SpecVersion >> 8), static_cast<uint16>(SpecVersion & 0xFF)};
-
-					LogVerbose(LOG_PREAMBLE, FString::Format(TEXT("Wintab Specification Version: {0}.{1}"),
-					                                         {InterfaceInfo.SpecificationVersion.Major, InterfaceInfo.SpecificationVersion.Minor}));
-				}
-
-				WORD ImplVersion;
-				if (ensure(WtInfo(WTI_INTERFACE, IFC_IMPLVERSION, &ImplVersion) == sizeof(ImplVersion)))
-				{
-					InterfaceInfo.ImplementationVersion = {static_cast<uint16>(ImplVersion >> 8), static_cast<uint16>(ImplVersion & 0xFF)};
-
-					LogVerbose(LOG_PREAMBLE, FString::Format(TEXT("Wintab Implementation Version: {0}.{1}"),
-					                                         {InterfaceInfo.ImplementationVersion.Major, InterfaceInfo.ImplementationVersion.Minor}));
-				}
-
-				/*UINT NumDevices;
-				if (ensure(WtInfo(WTI_INTERFACE, IFC_NDEVICES, &NumDevices) == sizeof(NumDevices)))
-				{
-					LogVerbose(LOG_PREAMBLE, FString::Format(TEXT("Wintab Number of Connected Devices: {0}"), {NumDevices}));
-				}*/
-
-				UINT MaxNumCursorTypes;
-				if (ensure(WtInfo(WTI_INTERFACE, IFC_NCURSORS, &MaxNumCursorTypes) == sizeof(MaxNumCursorTypes)))
-				{
-					InterfaceInfo.MaximumNumberOfCursorTypes = MaxNumCursorTypes;
-
-					LogVerbose(LOG_PREAMBLE, FString::Format(TEXT("Wintab Maximum Number of Cursor Types: {0}"), {MaxNumCursorTypes}));
-				}
-
-				UINT NumContexts;
-				if (ensure(WtInfo(WTI_INTERFACE, IFC_NCONTEXTS, &NumContexts) == sizeof(NumContexts)))
-				{
-					InterfaceInfo.MaximumNumberOfContexts = NumContexts;
-
-					LogVerbose(LOG_PREAMBLE, FString::Format(TEXT("Wintab Maximum Number of Contexts: {0}"), {NumContexts}));
-				}
-
-				UINT ContextOptions;
-				if (ensure(WtInfo(WTI_INTERFACE, IFC_CTXOPTIONS, &ContextOptions) == sizeof(ContextOptions)))
-				{
-					InterfaceInfo.bSupportsPacketMessages = ContextOptions & CXO_MESSAGES;
-					InterfaceInfo.bSupportsCursorChangeMessages = ContextOptions & CXO_CSRMESSAGES;
-
-					LogVerbose(LOG_PREAMBLE, FString::Format(TEXT("Wintab Supports System Cursor Contexts: {0}"), {ContextOptions & CXO_SYSTEM ? TEXT("Yes") : TEXT("No")}));
-					LogVerbose(LOG_PREAMBLE, FString::Format(TEXT("Wintab Supports Pen Windows Contexts: {0}"), {ContextOptions & CXO_PEN ? TEXT("Yes") : TEXT("No")}));
-					LogVerbose(LOG_PREAMBLE, FString::Format(TEXT("Wintab Supports WT_PACKET Messages: {0}"), {ContextOptions & CXO_MESSAGES ? TEXT("Yes") : TEXT("No")}));
-					LogVerbose(LOG_PREAMBLE, FString::Format(TEXT("Wintab Supports WT_CSRCHANGE  Messages: {0}"), {ContextOptions & CXO_CSRMESSAGES ? TEXT("Yes") : TEXT("No")}));
-					LogVerbose(LOG_PREAMBLE, FString::Format(TEXT("Wintab Supports Margins: {0}"), {ContextOptions & CXO_MARGIN ? TEXT("Yes") : TEXT("No")}));
-					LogVerbose(LOG_PREAMBLE, FString::Format(TEXT("Wintab Supports Inside Margins: {0}"), {ContextOptions & CXO_MGNINSIDE ? TEXT("Yes") : TEXT("No")}));
-				}
-			}
+			QueryInterfaceInfo();
 		}
 	}
 
@@ -140,20 +71,160 @@ namespace UE::StylusInput::Wintab
 		}
 	}
 
+	void FWintabAPI::QueryInterfaceInfo()
+	{
+		{
+			FWintabInfoOutputBuffer OutputBuffer;
+			TCHAR* const OutputBufferPtr = OutputBuffer.Allocate(WTI_INTERFACE, IFC_WINTABID);
+			const UINT BytesWritten = WtInfo(WTI_INTERFACE, IFC_WINTABID, OutputBufferPtr);
+			if (0 < BytesWritten && BytesWritten <= OutputBuffer.SizeInBytes())
+			{
+				InterfaceInfo.ID = OutputBufferPtr;
+
+				Log(LOG_PREAMBLE, FString::Format(TEXT("Interface ID: {0}"), {InterfaceInfo.ID}));
+			}
+			else if (BytesWritten > OutputBuffer.SizeInBytes())
+			{
+				LogError(LOG_PREAMBLE, FString::Format(
+					         TEXT("Failed to query tablet hardware identification string (IFC_WINTABID, return value {0})."),
+					         {BytesWritten}));
+			}
+		}
+
+		{
+			WORD SpecVersion;
+			const UINT BytesWritten = WtInfo(WTI_INTERFACE, IFC_SPECVERSION, &SpecVersion);
+			if (BytesWritten == sizeof(SpecVersion))
+			{
+				InterfaceInfo.SpecificationVersion = {static_cast<uint16>(SpecVersion >> 8), static_cast<uint16>(SpecVersion & 0xFF)};
+
+				Log(LOG_PREAMBLE, FString::Format(
+					    TEXT("Specification Version: {0}.{1}"),
+					           {InterfaceInfo.SpecificationVersion.Major, InterfaceInfo.SpecificationVersion.Minor}));
+			}
+			else if (BytesWritten != 0)
+			{
+				LogError(LOG_PREAMBLE, FString::Format(
+					         TEXT("Failed to query Specification Version (IFC_SPECVERSION, return value {0})."),
+					         {BytesWritten}));
+			}
+		}
+
+		{
+			WORD ImplVersion;
+			const UINT BytesWritten = WtInfo(WTI_INTERFACE, IFC_IMPLVERSION, &ImplVersion);
+			if (BytesWritten == sizeof(ImplVersion))
+			{
+				InterfaceInfo.ImplementationVersion = {static_cast<uint16>(ImplVersion >> 8), static_cast<uint16>(ImplVersion & 0xFF)};
+
+				Log(LOG_PREAMBLE, FString::Format(
+					           TEXT("Implementation Version: {0}.{1}"),
+					           {InterfaceInfo.ImplementationVersion.Major, InterfaceInfo.ImplementationVersion.Minor}));
+			}
+			else if (BytesWritten != 0)
+			{
+				LogError(LOG_PREAMBLE, FString::Format(
+					         TEXT("Failed to query Implementation Version (IFC_IMPLVERSION, return value {0})."),
+					         {BytesWritten}));
+			}
+		}
+
+		{
+			UINT MaxNumCursorTypes;
+			const UINT BytesWritten = WtInfo(WTI_INTERFACE, IFC_NCURSORS, &MaxNumCursorTypes);
+			if (BytesWritten == sizeof(MaxNumCursorTypes))
+			{
+				InterfaceInfo.MaximumNumberOfCursorTypes = MaxNumCursorTypes;
+
+				Log(LOG_PREAMBLE, FString::Format(TEXT("Maximum Number of Cursor Types: {0}"), {MaxNumCursorTypes}));
+			}
+			else if (BytesWritten != 0)
+			{
+				LogError(LOG_PREAMBLE, FString::Format(
+					         TEXT("Failed to query Maximum Number of Cursor Types (IFC_NCURSORS, return value {0})."),
+					         {BytesWritten}));
+			}
+		}
+
+		{
+			UINT NumContexts;
+			const UINT BytesWritten = WtInfo(WTI_INTERFACE, IFC_NCONTEXTS, &NumContexts);
+			if (BytesWritten == sizeof(NumContexts))
+			{
+				InterfaceInfo.MaximumNumberOfContexts = NumContexts;
+
+				Log(LOG_PREAMBLE, FString::Format(TEXT("Maximum Number of Contexts: {0}"), {InterfaceInfo.MaximumNumberOfContexts}));
+			}
+			else if (BytesWritten != 0)
+			{
+				LogError(LOG_PREAMBLE, FString::Format(
+					         TEXT("Failed to query maximum number of supported contexts (IFC_NCONTEXTS, return value {0})."),
+					         {BytesWritten}));
+			}
+		}
+
+		{
+			UINT ContextOptions;
+			const UINT BytesWritten = WtInfo(WTI_INTERFACE, IFC_CTXOPTIONS, &ContextOptions);
+			if (BytesWritten == sizeof(ContextOptions))
+			{
+				InterfaceInfo.bSupportsPacketMessages = ContextOptions & CXO_MESSAGES;
+				InterfaceInfo.bSupportsCursorChangeMessages = ContextOptions & CXO_CSRMESSAGES;
+
+				Log(LOG_PREAMBLE, FString::Format(
+					    TEXT("Supports System Cursor Contexts: {0}"), {ContextOptions & CXO_SYSTEM ? TEXT("Yes") : TEXT("No")}));
+				Log(LOG_PREAMBLE, FString::Format(
+					    TEXT("Supports Pen Windows Contexts: {0}"), {ContextOptions & CXO_PEN ? TEXT("Yes") : TEXT("No")}));
+				Log(LOG_PREAMBLE, FString::Format(
+					    TEXT("Supports WT_PACKET Messages: {0}"), {ContextOptions & CXO_MESSAGES ? TEXT("Yes") : TEXT("No")}));
+				Log(LOG_PREAMBLE, FString::Format(
+					    TEXT("Supports WT_CSRCHANGE  Messages: {0}"), {ContextOptions & CXO_CSRMESSAGES ? TEXT("Yes") : TEXT("No")}));
+				Log(LOG_PREAMBLE, FString::Format(
+					    TEXT("Supports Margins: {0}"), {ContextOptions & CXO_MARGIN ? TEXT("Yes") : TEXT("No")}));
+				Log(LOG_PREAMBLE, FString::Format(
+					    TEXT("Supports Inside Margins: {0}"), {ContextOptions & CXO_MGNINSIDE ? TEXT("Yes") : TEXT("No")}));
+			}
+			else if (BytesWritten != 0)
+			{
+				LogError(LOG_PREAMBLE, FString::Format(
+					         TEXT("Failed to query supported context options (IFC_CTXOPTIONS, return value {0})."),
+					         {BytesWritten}));
+			}
+		}
+
+		{
+			UINT NumDevices;
+			const UINT BytesWritten = WtInfo(WTI_INTERFACE, IFC_NDEVICES, &NumDevices);
+			if (BytesWritten == sizeof(NumDevices))
+			{
+				InterfaceInfo.NumberOfDevices = NumDevices;
+
+				Log(LOG_PREAMBLE, FString::Format(TEXT("Number of Devices: {0}"), {InterfaceInfo.NumberOfDevices}));
+			}
+			else if (BytesWritten != 0)
+			{
+				LogError(LOG_PREAMBLE, FString::Format(
+					         TEXT("Failed to query Number of Devices (IFC_NDEVICES, return value {0})."),
+					         {BytesWritten}));
+			}
+		}
+	}
+
 	bool FWintabAPI::DriverIsAvailable()
 	{
+		bool bIsAvailable = false;
+
 		if (void* Wintab32DllHandle = FWindowsPlatformProcess::GetDllHandle(TEXT("Wintab32.dll")))
 		{
-			if (FFuncWtInfo* WtInfo = static_cast<FFuncWtInfo*>(FWindowsPlatformProcess::GetDllExport(Wintab32DllHandle, TEXT("WTInfoW"))))
+			if (FWindowsPlatformProcess::GetDllExport(Wintab32DllHandle, TEXT("WTInfoW")))
 			{
-				if (WtInfo(0, 0, nullptr) != 0)
-				{
-					return true;
-				}
+				bIsAvailable = true;
 			}
+
+			FWindowsPlatformProcess::FreeDllHandle(Wintab32DllHandle);
 		}
 
-		return false;
+		return bIsAvailable;
 	}
 
 	FName FWintabAPI::GetName()
@@ -168,11 +239,10 @@ namespace UE::StylusInput::Wintab
 		return WintabAPI;
 	}
 
-	bool FWintabAPI::SatisfiesRequirements() const
+	bool FWintabAPI::IsValid() const
 	{
 		bool bSuccess = true;
 
-		bSuccess &= bWintabIsAvailable;
 		bSuccess &= GetWindowRect != nullptr;
 		bSuccess &= GetSystemMetrics != nullptr;
 		bSuccess &= WtInfo != nullptr;
@@ -182,6 +252,20 @@ namespace UE::StylusInput::Wintab
 
 		return bSuccess;
 	}
+
+	void FWintabAPI::UpdateNumberOfDevices() const
+	{
+		UINT NumDevices;
+		if (WtInfo(WTI_INTERFACE, IFC_NDEVICES, &NumDevices) == sizeof(NumDevices))
+		{
+			if (NumDevices != InterfaceInfo.NumberOfDevices)
+			{
+				InterfaceInfo.NumberOfDevices = NumDevices;
+
+				Log(LOG_PREAMBLE, FString::Format(TEXT("Wintab Number of Devices: {0}"), {InterfaceInfo.NumberOfDevices}));
+			}
+		}
+	}
 }
 
 #undef LOG_PREAMBLE
diff --git a/Engine/Plugins/Editor/StylusInput/Source/StylusInputWintab/Private/WintabAPI.h b/Engine/Plugins/Editor/StylusInput/Source/StylusInputWintab/Private/WintabAPI.h
--- a/Engine/Plugins/Editor/StylusInput/Source/StylusInputWintab/Private/WintabAPI.h
+++ b/Engine/Plugins/Editor/StylusInput/Source/StylusInputWintab/Private/WintabAPI.h
@@ -12,10 +12,11 @@
 	#include <windows.h>
 
 	#undef INT // WINTAB.H defines an INT macro that collides with a macro definition in AllowWindowsPlatformTypes.h when using Unity builds.
-		#include <WINTAB.H>
-		#define PACKETDATA (/*PK_CONTEXT |*/ PK_STATUS | PK_TIME | /*PK_CHANGED |*/ PK_SERIAL_NUMBER | PK_CURSOR | PK_BUTTONS | PK_X | PK_Y | PK_Z | PK_NORMAL_PRESSURE | PK_TANGENT_PRESSURE | PK_ORIENTATION | PK_ROTATION)
-		#define PACKETMODE PK_BUTTONS
-		#include <PKTDEF.H>
+
+	#include <WINTAB.H>
+	#define PACKETDATA (/*PK_CONTEXT |*/ PK_STATUS | PK_TIME | /*PK_CHANGED |*/ PK_SERIAL_NUMBER | PK_CURSOR | PK_BUTTONS | PK_X | PK_Y | PK_Z | PK_NORMAL_PRESSURE | PK_TANGENT_PRESSURE | PK_ORIENTATION | PK_ROTATION)
+	#define PACKETMODE PK_BUTTONS
+	#include <PKTDEF.H>
 
 #include <Windows/HideWindowsPlatformTypes.h>
 
@@ -38,6 +39,8 @@ namespace UE::StylusInput::Wintab
 		uint32 MaximumNumberOfContexts = 0;
 		bool bSupportsPacketMessages = false;
 		bool bSupportsCursorChangeMessages = false;
+
+		mutable uint32 NumberOfDevices = 0;
 	};
 
 	class FWintabAPI
@@ -48,7 +51,9 @@ namespace UE::StylusInput::Wintab
 		static FName GetName();
 		static const FWintabAPI& GetInstance();
 
-		bool SatisfiesRequirements() const;
+		bool IsValid() const;
+
+		void UpdateNumberOfDevices() const;
 
 		using FFuncGetWindowRect = BOOL(HWND, LPRECT);
 		FFuncGetWindowRect* GetWindowRect = nullptr;
@@ -75,9 +80,9 @@ namespace UE::StylusInput::Wintab
 		FWintabAPI();
 		~FWintabAPI();
 
-		TArray<void*> DllHandles;
+		void QueryInterfaceInfo();
 
-		bool bWintabIsAvailable = false;
+		TArray<void*> DllHandles;
 	};
 
 	inline uint32 HctxToUint32(const HCTX Hctx)
@@ -94,26 +99,31 @@ namespace UE::StylusInput::Wintab
 
 	/*
 	 * Apparently, some drivers have issues with a nullptr being passed as output buffer for WtInfo() calls.
-	 * The workaround is to just always have a fixed size buffer and hope for the best.
+	 * The workaround is to just always use a relatively large fixed size buffer and hope for the best.
 	 */
 
 	class FWintabInfoOutputBuffer
 	{
 	public:
-		void Allocate(const UINT Category, const UINT Index)
+		TCHAR* Allocate(UINT, UINT)
 		{
 			if (!Buffer.IsValid())
 			{
 				Buffer.Reset(new TCHAR[Size]);
 			}
+
 			Buffer[0] = 0;
+
+			return Buffer.Get();
 		}
 
-		TCHAR* Get() const { return Buffer.Get(); }
-		constexpr UINT SizeInBytes() { return Size * sizeof(TCHAR); }
+		static UINT SizeInBytes()
+		{
+			return Size * sizeof(TCHAR);
+		}
 
 	private:
-		enum { Size = 1024 };
+		enum { Size = 2048 };
 		TUniquePtr<TCHAR[]> Buffer;
 	};
 
@@ -127,7 +137,7 @@ namespace UE::StylusInput::Wintab
 		{
 		}
 
-		void Allocate(const UINT Category, const UINT Index)
+		TCHAR* Allocate(const UINT Category, const UINT Index)
 		{
 			// If the following call crashes, consider enabling the workaround available via WTINFO_WORKAROUND_FOR_OUTPUT_BUFFER_SIZE_BUG.
 			const UINT RequiredSize = WintabAPI.WtInfo(Category, Index, nullptr)  / sizeof(TCHAR);
@@ -137,14 +147,19 @@ namespace UE::StylusInput::Wintab
 				Size = RequiredSize;
 				Buffer.Reset(new TCHAR[Size]);
 			}
+
 			if (Size > 0)
 			{
 				Buffer[0] = 0;
 			}
+
+			return Buffer.Get();
 		}
 
-		TCHAR* Get() const { return Buffer.Get(); }
-		UINT SizeInBytes() const { return Size * sizeof(TCHAR); }
+		UINT SizeInBytes() const
+		{
+			return Size * sizeof(TCHAR);
+		}
 
 	private:
 		const FWintabAPI& WintabAPI;
diff --git a/Engine/Plugins/Editor/StylusInput/Source/StylusInputWintab/Private/WintabInstance.cpp b/Engine/Plugins/Editor/StylusInput/Source/StylusInputWintab/Private/WintabInstance.cpp
--- a/Engine/Plugins/Editor/StylusInput/Source/StylusInputWintab/Private/WintabInstance.cpp
+++ b/Engine/Plugins/Editor/StylusInput/Source/StylusInputWintab/Private/WintabInstance.cpp
@@ -65,10 +65,10 @@ namespace UE::StylusInput::Wintab
 		TabletContext.Name = WintabContext.lcName;
 		if (TabletContext.Name.IsEmpty())
 		{
-			OutputBuffer.Allocate(WTI_DEVICES + WintabContextIndex, DVC_NAME);
-			if (WintabAPI.WtInfo(WTI_DEVICES + WintabContextIndex, DVC_NAME, OutputBuffer.Get()) <= OutputBuffer.SizeInBytes())
+			TCHAR *const OutputBufferPtr = OutputBuffer.Allocate(WTI_DEVICES + WintabContextIndex, DVC_NAME);
+			if (WintabAPI.WtInfo(WTI_DEVICES + WintabContextIndex, DVC_NAME, OutputBufferPtr) <= OutputBuffer.SizeInBytes())
 			{
-				TabletContext.Name = OutputBuffer.Get();
+				TabletContext.Name = OutputBufferPtr;
 			}
 			else
 			{
@@ -76,10 +76,10 @@ namespace UE::StylusInput::Wintab
 			}
 		}
 
-		OutputBuffer.Allocate(WTI_DEVICES + WintabContextIndex, DVC_PNPID);
-		if (WintabAPI.WtInfo(WTI_DEVICES + WintabContextIndex, DVC_PNPID, OutputBuffer.Get()) <= OutputBuffer.SizeInBytes())
+		TCHAR *const OutputBufferPtr = OutputBuffer.Allocate(WTI_DEVICES + WintabContextIndex, DVC_PNPID);
+		if (WintabAPI.WtInfo(WTI_DEVICES + WintabContextIndex, DVC_PNPID, OutputBufferPtr) <= OutputBuffer.SizeInBytes())
 		{
-			TabletContext.PlugAndPlayID = OutputBuffer.Get();
+			TabletContext.PlugAndPlayID = OutputBufferPtr;
 		}
 		else
 		{
@@ -576,14 +576,14 @@ namespace UE::StylusInput::Wintab
 
 		if (const uint32 NumInvalidPackets = PacketStats.GetNumInvalidPackets())
 		{
-			LogVerbose(LOG_PREAMBLE, FString::Format(
-						   TEXT("Wintab instance '{0}' encountered {1} invalid packets."), {FWintabInstance::GetName().ToString(), NumInvalidPackets}));
+			Log(LOG_PREAMBLE, FString::Format(
+				    TEXT("Wintab instance '{0}' encountered {1} invalid packets."), {FWintabInstance::GetName().ToString(), NumInvalidPackets}));
 		}
 
 		if (const uint32 NumLostPackets = PacketStats.GetNumLostPackets())
 		{
-			LogVerbose(LOG_PREAMBLE, FString::Format(
-						   TEXT("Wintab instance '{0}' encountered {1} lost packets."), {FWintabInstance::GetName().ToString(), NumLostPackets}));
+			Log(LOG_PREAMBLE, FString::Format(
+				    TEXT("Wintab instance '{0}' encountered {1} lost packets."), {FWintabInstance::GetName().ToString(), NumLostPackets}));
 		}
 	}
 
@@ -653,7 +653,7 @@ namespace UE::StylusInput::Wintab
 		return FText::Format(LOCTEXT("Wintab", "Wintab #{0}"), ID);
 	}
 
-	bool FWintabInstance::WasInitializedSuccessfully()
+	bool FWintabInstance::WasInitializedSuccessfully() const
 	{
 		return true;
 	}
@@ -721,10 +721,10 @@ namespace UE::StylusInput::Wintab
 				}
 
 				FWintabInfoOutputBuffer OutputBuffer;
-				OutputBuffer.Allocate(WTI_CURSORS + CursorIndex, CSR_NAME);
-				if (WintabAPI.WtInfo(WTI_CURSORS + CursorIndex, CSR_NAME, OutputBuffer.Get()) <= OutputBuffer.SizeInBytes())
+				TCHAR* OutputBufferPtr = OutputBuffer.Allocate(WTI_CURSORS + CursorIndex, CSR_NAME);
+				if (WintabAPI.WtInfo(WTI_CURSORS + CursorIndex, CSR_NAME, OutputBufferPtr) <= OutputBuffer.SizeInBytes())
 				{
-					StylusInfo->Name = FString::Format(TEXT("{0} ({1})"), {OutputBuffer.Get(), MaskedCursorTypeToString(MaskedCursorType)});
+					StylusInfo->Name = FString::Format(TEXT("{0} ({1})"), {OutputBufferPtr, MaskedCursorTypeToString(MaskedCursorType)});
 				}
 				else
 				{
@@ -736,10 +736,10 @@ namespace UE::StylusInput::Wintab
 				{
 					StylusInfo->Buttons.SetNum(NumButtons);
 
-					OutputBuffer.Allocate(WTI_CURSORS + CursorIndex, CSR_BTNNAMES);
-					if (WintabAPI.WtInfo(WTI_CURSORS + CursorIndex, CSR_BTNNAMES, OutputBuffer.Get()) <= OutputBuffer.SizeInBytes())
+					OutputBufferPtr = OutputBuffer.Allocate(WTI_CURSORS + CursorIndex, CSR_BTNNAMES);
+					if (WintabAPI.WtInfo(WTI_CURSORS + CursorIndex, CSR_BTNNAMES, OutputBufferPtr) <= OutputBuffer.SizeInBytes())
 					{
-						const TCHAR* Name = OutputBuffer.Get();
+						const TCHAR* Name = OutputBufferPtr;
 						int32 ButtonIndex = 0;
 						while (*Name)
 						{
diff --git a/Engine/Plugins/Editor/StylusInput/Source/StylusInputWintab/Private/WintabInstance.h b/Engine/Plugins/Editor/StylusInput/Source/StylusInputWintab/Private/WintabInstance.h
--- a/Engine/Plugins/Editor/StylusInput/Source/StylusInputWintab/Private/WintabInstance.h
+++ b/Engine/Plugins/Editor/StylusInput/Source/StylusInputWintab/Private/WintabInstance.h
@@ -36,7 +36,7 @@ namespace UE::StylusInput::Wintab
 		virtual FName GetInterfaceName() override;
 		virtual FText GetName() override;
 
-		virtual bool WasInitializedSuccessfully() override;
+		virtual bool WasInitializedSuccessfully() const override;
 
 	private:
 		const FTabletContext* GetTabletContextInternal(uint32 TabletContextID) const;
diff --git a/Engine/Plugins/Editor/StylusInput/Source/StylusInputWintab/Private/WintabInterface.cpp b/Engine/Plugins/Editor/StylusInput/Source/StylusInputWintab/Private/WintabInterface.cpp
--- a/Engine/Plugins/Editor/StylusInput/Source/StylusInputWintab/Private/WintabInterface.cpp
+++ b/Engine/Plugins/Editor/StylusInput/Source/StylusInputWintab/Private/WintabInterface.cpp
@@ -85,7 +85,7 @@ namespace UE::StylusInput::Wintab
 
 	TUniquePtr<IStylusInputInterface> FWintabInterface::Create()
 	{
-		if (FWintabAPI::GetInstance().SatisfiesRequirements())
+		if (FWintabAPI::GetInstance().IsValid())
 		{
 			TUniquePtr<FWintabInterface> WintabInterface = MakeUnique<FWintabInterface>();
 			return TUniquePtr<IStylusInputInterface>(MoveTemp(WintabInterface));
diff --git a/Engine/Plugins/Editor/StylusInput/Source/StylusInputWintab/Private/WintabMessageHandler.cpp b/Engine/Plugins/Editor/StylusInput/Source/StylusInputWintab/Private/WintabMessageHandler.cpp
--- a/Engine/Plugins/Editor/StylusInput/Source/StylusInputWintab/Private/WintabMessageHandler.cpp
+++ b/Engine/Plugins/Editor/StylusInput/Source/StylusInputWintab/Private/WintabMessageHandler.cpp
@@ -156,6 +156,8 @@ bool FWintabMessageHandler::ProcessMessage(const HWND Hwnd, const uint32 Msg, co
 			DebugEvent(FString::Format(TEXT("Received WT_INFOCHANGE(Context={0}, Category={1}, Index={2}) message."), {WParam, LOWORD(LParam), HIWORD(LParam)}));
 #endif
 
+			WintabAPI.UpdateNumberOfDevices();
+
 			UpdateTabletContextsCallback.Execute();
 
 			bHandled = true;
