StylusInput: MacOS support

#rn StylusInput: Added MacOS support

diff --git a/Engine/Plugins/Editor/StylusInput/Source/StylusInputMac/Private/MacInstance.cpp b/Engine/Plugins/Editor/StylusInput/Source/StylusInputMac/Private/MacInstance.cpp
new file mode 100644
--- /dev/null
+++ b/Engine/Plugins/Editor/StylusInput/Source/StylusInputMac/Private/MacInstance.cpp
@@ -0,0 +1,153 @@
+// Copyright Epic Games, Inc. All Rights Reserved.
+
+#include "MacInstance.h"
+
+#include <StylusInput.h>
+#include <StylusInputUtils.h>
+#include <Algo/Transform.h>
+
+#include "MacInterface.h"
+#include "NSEventHandler.h"
+#include "MacTabletContext.h"
+
+#define LOCTEXT_NAMESPACE "MacInstance"
+#define LOG_PREAMBLE "MacInstance"
+
+using namespace UE::StylusInput::Private;
+
+namespace UE::StylusInput::Mac
+{
+	FMacInstance::FMacInstance(uint32 ID, const FCocoaWindow* OSWindowHandle)
+		: CocoaWindow(OSWindowHandle)
+		, ID(ID)
+	{
+	}
+
+	FMacInstance::~FMacInstance()
+	{
+	}
+
+	bool FMacInstance::AddEventHandler(IStylusInputEventHandler* EventHandler, const EEventHandlerThread Thread)
+	{
+		if (!EventHandler)
+		{
+			LogWarning(LOG_PREAMBLE, "Tried to add nullptr as event handler.");
+			return false;
+		}
+		
+		if (!NSEventHandler)
+		{
+			EnableEventHandler(EEventHandlerThread::OnGameThread, EventHandler);
+			
+			if (!NSEventHandler)
+			{
+				LogError(LOG_PREAMBLE, FString::Format(
+						 TEXT("Event handler '{0}' was not added since NSEventHandler could not be installed."),
+						 {EventHandler->GetName()}));
+				return false;
+			}
+			return true;
+		}
+		
+		return NSEventHandler->AddEventHandler(EventHandler);
+	}
+
+	bool FMacInstance::RemoveEventHandler(IStylusInputEventHandler* EventHandler)
+	{
+		if (!EventHandler)
+		{
+			LogWarning(LOG_PREAMBLE, "Tried to remove nullptr event handler.");
+			return false;
+		}
+
+		bool bWasRemoved = false;
+
+		if (NSEventHandler && NSEventHandler->RemoveEventHandler(EventHandler))
+		{
+			bWasRemoved = true;
+
+			if (NSEventHandler->NumEventHandlers() == 0)
+			{
+				DisableEventHandler(EEventHandlerThread::OnGameThread);
+			}
+		}
+
+		if (!bWasRemoved)
+		{
+			LogError(LOG_PREAMBLE, FString::Format(TEXT("Event handler '%s' does not exist."), {EventHandler->GetName()}));
+		}
+
+		return bWasRemoved;
+	}
+
+	const TSharedPtr<IStylusInputTabletContext> FMacInstance::GetTabletContext(const uint32 TabletContextID)
+	{
+		return TabletContexts.Get(TabletContextID);
+	}
+
+	const TSharedPtr<IStylusInputStylusInfo> FMacInstance::GetStylusInfo(uint32 StylusID)
+	{
+		return StylusInfo.Get(StylusID);
+	}
+
+	float FMacInstance::GetPacketsPerSecond(const EEventHandlerThread EventHandlerThread) const
+	{
+		return NSEventHandler ? NSEventHandler->GetPacketsPerSecond() : 0.0f;
+	}
+
+	void FMacInstance::SetupTabletContexts()
+	{
+		if( NSEventHandler )
+		{
+			GetMacTabletDevices(NSEventHandler->GetTabletContexts());
+			for (int i = 0; i < NSEventHandler->GetTabletContexts().Num(); i++)
+			{
+				LogVerbose(LOG_PREAMBLE, FString::Format(TEXT("Added tablet context for device: %s, productID: %d, vendorID: %d"), {*NSEventHandler->GetTabletContexts()[i]->ProductName, NSEventHandler->GetTabletContexts()[i]->ProductID, NSEventHandler->GetTabletContexts()[i]->VendorID}));
+			}
+			UpdateTabletContexts(NSEventHandler->GetTabletContexts());
+		}
+	}
+
+	void FMacInstance::UpdateTabletContexts(FTabletContextContainer& InTabletContexts)
+	{
+		TabletContexts.Update(InTabletContexts);
+	}
+
+	void FMacInstance::EnableEventHandler(const EEventHandlerThread EventHandlerThread, IStylusInputEventHandler* EventHandler)
+	{
+		NSEventHandler = MakeUnique<FNSEventHandler>(
+			this,
+			CocoaWindow,
+			EventHandler
+		);
+		
+		SetupTabletContexts();
+		NSEventHandler->StartListen();
+	}
+
+	void FMacInstance::DisableEventHandler(const EEventHandlerThread EventHandlerThread)
+	{
+		if( NSEventHandler )
+			NSEventHandler->StopListen();
+		NSEventHandler.Reset();
+	}
+
+	FName FMacInstance::GetInterfaceName()
+	{
+		static FName Name("NSEvent");
+		return Name;
+	}
+
+	FText FMacInstance::GetName()
+	{
+		return FText::Format(LOCTEXT("NSEvent", "NSEvent #{0}"), ID);
+	}
+
+	bool FMacInstance::WasInitializedSuccessfully() const
+	{
+		return true;
+	}
+}
+
+#undef LOG_PREAMBLE
+#undef LOCTEXT_NAMESPACE
diff --git a/Engine/Plugins/Editor/StylusInput/Source/StylusInputMac/Private/MacInstance.h b/Engine/Plugins/Editor/StylusInput/Source/StylusInputMac/Private/MacInstance.h
new file mode 100644
--- /dev/null
+++ b/Engine/Plugins/Editor/StylusInput/Source/StylusInputMac/Private/MacInstance.h
@@ -0,0 +1,59 @@
+// Copyright Epic Games, Inc. All Rights Reserved.
+
+#pragma once
+
+#include <StylusInput.h>
+#include <Containers/Array.h>
+#include <Templates/UniquePtr.h>
+
+#include "MacTabletContext.h"
+#include "Mac/MacApplication.h"
+#include "Mac/MacPlatformApplicationMisc.h"
+
+namespace UE::StylusInput
+{
+	class IStylusInputEventHandler;
+}
+
+namespace UE::StylusInput::Mac
+{
+	class FNSEventHandler;
+
+	class FMacInstance : public IStylusInputInstance
+	{
+	public:
+		explicit FMacInstance(uint32 ID, const FCocoaWindow* OSWindowHandle);
+		virtual ~FMacInstance() override;
+
+		virtual bool AddEventHandler(IStylusInputEventHandler* EventHandler, EEventHandlerThread Thread) override;
+		virtual bool RemoveEventHandler(IStylusInputEventHandler* EventHandler) override;
+
+		virtual const TSharedPtr<IStylusInputTabletContext> GetTabletContext(uint32 TabletContextID) override;
+		virtual const TSharedPtr<IStylusInputStylusInfo> GetStylusInfo(uint32 StylusID) override;
+
+		virtual float GetPacketsPerSecond(EEventHandlerThread EventHandlerThread) const override;
+
+		virtual FName GetInterfaceName() override;
+		virtual FText GetName() override;
+
+		virtual bool WasInitializedSuccessfully() const override;
+
+	private:
+		void EnableEventHandler(EEventHandlerThread EventHandlerThread, IStylusInputEventHandler* EventHandler);
+		void DisableEventHandler(EEventHandlerThread EventHandlerThread);
+
+		// There isn't such a thing as a "Tablet context" with NSEvent. Instead, we'll store all the devices known by the user computer and sort the events by ID
+		void SetupTabletContexts();
+
+		void UpdateTabletContexts(FTabletContextContainer& InTabletContexts);
+
+		const FCocoaWindow* CocoaWindow;
+
+		TUniquePtr<FNSEventHandler> NSEventHandler;
+
+		FTabletContextContainer TabletContexts;
+		FStylusInfoContainer StylusInfo;
+
+		const uint32 ID;
+	};
+}
diff --git a/Engine/Plugins/Editor/StylusInput/Source/StylusInputMac/Private/MacInterface.cpp b/Engine/Plugins/Editor/StylusInput/Source/StylusInputMac/Private/MacInterface.cpp
new file mode 100644
--- /dev/null
+++ b/Engine/Plugins/Editor/StylusInput/Source/StylusInputMac/Private/MacInterface.cpp
@@ -0,0 +1,89 @@
+// Copyright Epic Games, Inc. All Rights Reserved.
+
+#include "MacInterface.h"
+
+#include <StylusInput.h>
+#include <StylusInputUtils.h>
+
+#include <GenericPlatform/GenericWindow.h>
+#include <Templates/SharedPointer.h>
+#include <Widgets/SWindow.h>
+
+#include "MacInstance.h"
+
+using namespace UE::StylusInput::Private;
+
+namespace UE::StylusInput::Mac
+{
+	FName FMacInterface::GetName() const
+	{
+		static FName Name("NSEvent");
+		return Name;
+	}
+
+	IStylusInputInstance* FMacInterface::CreateInstance(SWindow& Window)
+	{
+		if (FRefCountedInstance* ExistingRefCountedInstance = Instances.Find(&Window))
+		{
+			++ExistingRefCountedInstance->RefCount;
+			return ExistingRefCountedInstance->Instance.Get();
+		}
+
+		FCocoaWindow* OSWindowHandle = [&Window]
+		{
+			const TSharedPtr<const FGenericWindow> NativeWindow = Window.GetNativeWindow();
+			return NativeWindow.IsValid() ? static_cast<FCocoaWindow*>(NativeWindow->GetOSWindowHandle()) : nullptr;
+		}();
+
+		if (!OSWindowHandle)
+		{
+			LogError("MacInterface", "Could not get native window handle.");
+			return nullptr;
+		}
+
+		FMacInstance* NewInstance = Instances.Emplace(
+			&Window, {MakeUnique<FMacInstance>(NextInstanceID++, OSWindowHandle), 1}).Instance.Get();
+		if (!ensureMsgf(NewInstance, TEXT("MacInterface: Failed to create stylus input instance.")))
+		{
+			Instances.Remove(&Window);
+			return nullptr;
+		}
+
+		return NewInstance;
+	}
+
+	bool FMacInterface::ReleaseInstance(IStylusInputInstance* Instance)
+	{
+		check(Instance);
+
+		FMacInstance *const MacInstance = static_cast<FMacInstance*>(Instance);
+
+		// Find existing instance
+		for (TTuple<SWindow*, FRefCountedInstance>& Entry : Instances)
+		{
+			FRefCountedInstance& RefCountedInstance = Entry.Get<1>();
+
+			if (RefCountedInstance.Instance.Get() == MacInstance)
+			{
+				// Decrease reference count
+				check(RefCountedInstance.RefCount > 0);
+				if (--RefCountedInstance.RefCount == 0)
+				{
+					// Delete if there are no references left
+					Instances.Remove(Entry.Key);
+				}
+
+				return true;
+			}
+		}
+
+		ensureMsgf(false, TEXT("MacInterface: Failed to find provided instance."));
+		return false;
+	}
+
+	TUniquePtr<IStylusInputInterface> FMacInterface::Create()
+	{
+		TUniquePtr<FMacInterface> MacImpl = MakeUnique<FMacInterface>();
+		return TUniquePtr<IStylusInputInterface>(MoveTemp(MacImpl));
+	}
+}
diff --git a/Engine/Plugins/Editor/StylusInput/Source/StylusInputMac/Private/MacInterface.h b/Engine/Plugins/Editor/StylusInput/Source/StylusInputMac/Private/MacInterface.h
new file mode 100644
--- /dev/null
+++ b/Engine/Plugins/Editor/StylusInput/Source/StylusInputMac/Private/MacInterface.h
@@ -0,0 +1,41 @@
+// Copyright Epic Games, Inc. All Rights Reserved.
+
+#pragma once
+
+#include <StylusInputInterface.h>
+#include <Containers/Array.h>
+#include <Containers/Map.h>
+#include <Templates/UniquePtr.h>
+
+namespace UE::StylusInput::Private
+{
+	class FStylusInputImpl;
+}
+
+namespace UE::StylusInput::Mac
+{
+	class FMacInstance;
+
+	class FMacInterface : public IStylusInputInterface
+	{
+	public:
+		virtual FName GetName() const override;
+
+		virtual IStylusInputInstance* CreateInstance(SWindow& Window) override;
+		virtual bool ReleaseInstance(IStylusInputInstance* Instance) override;
+
+	private:
+		friend class FStylusInputMacModule;
+
+		static TUniquePtr<IStylusInputInterface> Create();
+
+		struct FRefCountedInstance
+		{
+			TUniquePtr<FMacInstance> Instance;
+			int32 RefCount;
+		};
+
+		TMap<SWindow*, FRefCountedInstance> Instances;
+		uint32 NextInstanceID = 0;
+	};
+}
diff --git a/Engine/Plugins/Editor/StylusInput/Source/StylusInputMac/Private/MacStats.h b/Engine/Plugins/Editor/StylusInput/Source/StylusInputMac/Private/MacStats.h
new file mode 100644
--- /dev/null
+++ b/Engine/Plugins/Editor/StylusInput/Source/StylusInputMac/Private/MacStats.h
@@ -0,0 +1,40 @@
+// Copyright Epic Games, Inc. All Rights Reserved.
+
+#pragma once
+
+#include <Mac/MacPlatformTime.h>
+
+namespace UE::StylusInput::Mac
+{
+	class FPacketStats
+	{
+	public:
+		void NewPacket()
+		{
+			const uint64 Timestamp = FPlatformTime::Cycles64();
+
+			if (LatestTimestamp - EarliestTimestamp >= CyclesPerSecond)
+			{
+				PacketsPerSecond = EarliestTimestamp != LatestTimestamp ? NumPackets / ((LatestTimestamp - EarliestTimestamp) / CyclesPerSecond) : 0.0f;
+				EarliestTimestamp = Timestamp;
+				NumPackets = -1;
+			}
+
+			LatestTimestamp = Timestamp;
+			++NumPackets;
+		}
+
+		float GetPacketsPerSecond() const
+		{
+			const uint64 Timestamp = FPlatformTime::Cycles64();
+			return Timestamp - LatestTimestamp <= CyclesPerSecond ? PacketsPerSecond : 0.0f;
+		}
+
+	private:
+		uint64 EarliestTimestamp = 0;
+		uint64 LatestTimestamp = 0;
+		uint32 NumPackets = 0;
+		float PacketsPerSecond = 0.0f;
+		const double CyclesPerSecond = 1.0 / FPlatformTime::GetSecondsPerCycle64();
+	};
+}
diff --git a/Engine/Plugins/Editor/StylusInput/Source/StylusInputMac/Private/MacTabletContext.cpp b/Engine/Plugins/Editor/StylusInput/Source/StylusInputMac/Private/MacTabletContext.cpp
new file mode 100644
--- /dev/null
+++ b/Engine/Plugins/Editor/StylusInput/Source/StylusInputMac/Private/MacTabletContext.cpp
@@ -0,0 +1,220 @@
+// Copyright Epic Games, Inc. All Rights Reserved.
+
+#include "MacTabletContext.h"
+
+#import <IOKit/hid/IOHIDManager.h>
+#import <Foundation/Foundation.h>
+
+#define LOG_PREAMBLE "MacTabletContext"
+
+using namespace UE::StylusInput::Private;
+
+namespace UE::StylusInput::Mac
+{
+	static FString CFStringToFString(CFStringRef cfStr)
+	{
+		if (!cfStr) return FString();
+		NSString* nsStr = (__bridge NSString*)cfStr;
+		return FString(nsStr);
+	}
+
+	bool GetMacTabletDevices(FTabletContextContainer& OutDevices)
+	{
+		OutDevices.Clear();
+
+		IOHIDManagerRef HidManager = IOHIDManagerCreate(kCFAllocatorDefault, kIOHIDOptionsTypeNone);
+		if (!HidManager)
+		{
+			return false;
+		}
+
+		// Match Generic Desktop devices (mouse and such)
+		NSDictionary* matchingGenericDesktopDict = @{
+			@(kIOHIDDeviceUsagePageKey): @(kHIDPage_GenericDesktop)
+		};
+
+		// Match Digitizer devices (tablets)
+		NSDictionary* matchingDigitizerDict = @{
+			@(kIOHIDDeviceUsagePageKey): @(kHIDPage_Digitizer)
+		};
+
+		NSArray* matchingArray = @[matchingGenericDesktopDict, matchingDigitizerDict];
+
+		IOHIDManagerSetDeviceMatchingMultiple(HidManager, (__bridge CFArrayRef) matchingArray);
+		IOHIDManagerScheduleWithRunLoop(HidManager, CFRunLoopGetCurrent(), kCFRunLoopDefaultMode);
+		IOHIDManagerOpen(HidManager, kIOHIDOptionsTypeNone);
+
+		CFSetRef DeviceSet = IOHIDManagerCopyDevices(HidManager);
+		if (!DeviceSet || CFSetGetCount(DeviceSet) == 0)
+		{
+			if (DeviceSet) 
+				CFRelease(DeviceSet);
+			CFRelease(HidManager);
+			return false;
+		}
+
+		NSSet* Devices = (__bridge NSSet*)DeviceSet;
+
+		for (id DeviceObj in Devices)
+		{
+			IOHIDDeviceRef Device = (__bridge IOHIDDeviceRef)DeviceObj;
+
+			//---
+			CFStringRef productKey = (CFStringRef)IOHIDDeviceGetProperty(Device, CFSTR(kIOHIDProductKey));
+			NSNumber* deviceUsagePageNumber = (__bridge NSNumber*)IOHIDDeviceGetProperty(Device, CFSTR(kIOHIDPrimaryUsagePageKey));
+			NSNumber* deviceUsageNumber = (__bridge NSNumber*)IOHIDDeviceGetProperty(Device, CFSTR(kIOHIDPrimaryUsageKey));
+			uint32_t deviceUsagePage = deviceUsagePageNumber.unsignedIntValue;
+			uint32_t deviceUsage = deviceUsageNumber.unsignedIntValue;
+
+			NSLog(@"Supported Device: %@ | UsagePage: %@, Usage: %@", productKey, deviceUsagePageNumber, deviceUsageNumber);
+			//---
+
+			CFStringRef product = (CFStringRef)IOHIDDeviceGetProperty(Device, CFSTR(kIOHIDProductKey));
+			CFStringRef vendor = (CFStringRef)IOHIDDeviceGetProperty(Device, CFSTR(kIOHIDManufacturerKey));
+			CFStringRef serial = (CFStringRef)IOHIDDeviceGetProperty(Device, CFSTR(kIOHIDSerialNumberKey));
+			CFNumberRef vendorID = (CFNumberRef)IOHIDDeviceGetProperty(Device, CFSTR(kIOHIDVendorIDKey));
+			CFNumberRef productID = (CFNumberRef)IOHIDDeviceGetProperty(Device, CFSTR(kIOHIDProductIDKey));
+
+			NSLog(@"Supported product: %@ | vendor: %@, serial: %@, vendorID: %@, productID: %@", product, vendor, serial, vendorID, productID);
+
+			uint32 ID = 0;
+			if(productID) //productID is 0 if and only if we're dealing with the internal pointing device (touchpad)
+				CFNumberGetValue(productID, kCFNumberSInt32Type, &ID);
+
+			LogVerbose(LOG_PREAMBLE, FString::Format(TEXT("ID: %d"), {ID}));
+
+			TSharedRef<FTabletContext> context = OutDevices.Add(ID);
+
+			context->ProductName = product;
+			context->ProductID = ID;
+
+			context->VendorName  = vendor;
+			if (vendorID)
+				CFNumberGetValue(vendorID, kCFNumberSInt32Type, &context->VendorID);
+
+			context->SerialNumber = serial;
+
+			CFArrayRef elements = IOHIDDeviceCopyMatchingElements(Device, NULL, kIOHIDOptionsTypeNone);
+			if (elements) 
+			{
+				for (CFIndex i = 0; i < CFArrayGetCount(elements); i++) 
+				{
+					IOHIDElementRef element = (IOHIDElementRef)CFArrayGetValueAtIndex(elements, i);
+					uint32_t elementUsagePage = IOHIDElementGetUsagePage(element);
+					uint32_t elementUsage = IOHIDElementGetUsage(element);
+					NSLog(@"Supported element usage page 0x%X and usage: 0x%X", elementUsagePage, elementUsage);
+
+					if( elementUsagePage == kHIDPage_GenericDesktop )
+					{
+						switch(elementUsage)
+						{
+							case kHIDUsage_GD_Mouse:
+								context->HardwareCapabilities = context->HardwareCapabilities | ETabletHardwareCapabilities::Integrated;
+
+								NSLog(@"Supported Integrated");
+								break;
+							case kHIDUsage_GD_X:
+								context->SupportedProperties = context->SupportedProperties | ETabletSupportedProperties::X;
+								context->InputRectangle.Min[0] = IOHIDElementGetLogicalMin(element);
+								context->InputRectangle.Max[0] = IOHIDElementGetLogicalMax(element) - context->InputRectangle.Min[0];
+								NSLog(@"Supported X");
+								break;
+							case kHIDUsage_GD_Y:
+								context->SupportedProperties = context->SupportedProperties | ETabletSupportedProperties::Y;
+								context->InputRectangle.Min[1] = IOHIDElementGetLogicalMin(element);
+								context->InputRectangle.Max[1] = IOHIDElementGetLogicalMax(element) - context->InputRectangle.Min[1];
+								NSLog(@"Supported Y");
+								break;
+							case kHIDUsage_GD_Z:
+								context->SupportedProperties = context->SupportedProperties | ETabletSupportedProperties::Z;
+								NSLog(@"Supported Z");
+								break;
+						}
+					}
+					if( elementUsagePage == kHIDPage_Digitizer || elementUsagePage >= 0xFF00 ) //Vendor-defined enum field
+					{
+						switch(elementUsage)
+						{
+							case kHIDUsage_Dig_Pen:
+							case kHIDUsage_Dig_LightPen:
+							case kHIDUsage_Dig_Stylus:
+								context->HardwareCapabilities = context->HardwareCapabilities | ETabletHardwareCapabilities::CursorMustTouch;
+								NSLog(@"CursorMustTouch Supported");
+								break;
+							case kHIDUsage_Dig_TouchScreen:
+							case kHIDUsage_Dig_TouchPad:
+								context->HardwareCapabilities = context->HardwareCapabilities | ETabletHardwareCapabilities::Integrated;
+								NSLog(@"Dig Integrated Supported");
+								break;
+							case kHIDUsage_Dig_InRange:
+								context->HardwareCapabilities = context->HardwareCapabilities | ETabletHardwareCapabilities::HardProximity;
+								NSLog(@"Dig HardProximity Supported");
+								break;
+							case kHIDUsage_Dig_BarrelSwitch:
+							case kHIDUsage_Dig_Eraser:
+								context->SupportedProperties = context->SupportedProperties | ETabletSupportedProperties::PacketStatus;
+								NSLog(@"PacketStatus Supported");
+								break;
+							case kHIDUsage_Dig_Touch:
+								context->SupportedProperties = context->SupportedProperties | ETabletSupportedProperties::PacketStatus;
+								context->HardwareCapabilities = context->HardwareCapabilities | ETabletHardwareCapabilities::Integrated;
+								NSLog(@"Touch Supported");
+							case kHIDUsage_Dig_TipPressure:
+								context->SupportedProperties = context->SupportedProperties | ETabletSupportedProperties::NormalPressure;
+								NSLog(@"Pressure Supported");
+								break;
+							case kHIDUsage_Dig_BarrelPressure:
+								context->SupportedProperties = context->SupportedProperties | ETabletSupportedProperties::TangentPressure;
+								NSLog(@"Tangent Pressure Supported");
+								break;
+							case kHIDUsage_Dig_XTilt:
+								context->SupportedProperties = context->SupportedProperties | ETabletSupportedProperties::XTiltOrientation;
+								NSLog(@"XTilt Supported");
+								break;
+							case kHIDUsage_Dig_YTilt:
+								context->SupportedProperties = context->SupportedProperties | ETabletSupportedProperties::YTiltOrientation;
+								NSLog(@"YTilt Supported");
+								break;
+							case kHIDUsage_Dig_Azimuth:
+								context->SupportedProperties = context->SupportedProperties | ETabletSupportedProperties::AzimuthOrientation;
+								NSLog(@"Azimuth Supported");
+								break;
+							case kHIDUsage_Dig_Altitude:
+								context->SupportedProperties = context->SupportedProperties | ETabletSupportedProperties::AltitudeOrientation;
+								NSLog(@"Altitude Supported");
+								break;
+							case kHIDUsage_Dig_Twist:
+								context->SupportedProperties = context->SupportedProperties | ETabletSupportedProperties::TwistOrientation;
+								NSLog(@"Twist Supported");
+								break;
+							case kHIDUsage_Dig_Width:
+								context->SupportedProperties = context->SupportedProperties | ETabletSupportedProperties::Width;
+								NSLog(@"Width Supported");
+								break;
+							case kHIDUsage_Dig_Height:
+								context->SupportedProperties = context->SupportedProperties | ETabletSupportedProperties::Height;
+								NSLog(@"Height Supported");
+								break;
+							case kHIDUsage_Dig_Finger:
+								context->SupportedProperties = context->SupportedProperties | ETabletSupportedProperties::FingerContactConfidence;
+								context->HardwareCapabilities = context->HardwareCapabilities | ETabletHardwareCapabilities::Integrated;
+								NSLog(@"Finger Confidence Supported");
+								break;
+							case kHIDUsage_Dig_ContactIdentifier:
+								context->SupportedProperties = context->SupportedProperties | ETabletSupportedProperties::DeviceContactID;
+								context->HardwareCapabilities = context->HardwareCapabilities | ETabletHardwareCapabilities::CursorsHavePhysicalIds;
+								NSLog(@"Device contact ID Supported");
+								break;
+						}
+					}
+				}
+			}
+		}
+		CFRelease(DeviceSet);
+		CFRelease(HidManager);
+
+		return true;
+	}
+}
+
+#undef LOG_PREAMBLE
diff --git a/Engine/Plugins/Editor/StylusInput/Source/StylusInputMac/Private/MacTabletContext.h b/Engine/Plugins/Editor/StylusInput/Source/StylusInputMac/Private/MacTabletContext.h
new file mode 100644
--- /dev/null
+++ b/Engine/Plugins/Editor/StylusInput/Source/StylusInputMac/Private/MacTabletContext.h
@@ -0,0 +1,155 @@
+// Copyright Epic Games, Inc. All Rights Reserved.
+
+#pragma once
+
+#include <StylusInputPacket.h>
+#include <StylusInputTabletContext.h>
+#include <StylusInputUtils.h>
+#include <Math/IntRect.h>
+#include <Math/Vector.h>
+
+namespace UE::StylusInput::Mac
+{
+	enum class EPacketPropertyType : int8
+	{
+		X = 0, // The x-coordinate in the tablet coordinate space. Each packet contains this property by default. The origin (0,0) of the tablet is the upper-left corner.
+		Y, // The y-coordinate in the tablet coordinate space. Each packet contains this property by default. The origin (0,0) of the tablet is the upper-left corner.
+		Z, // The z-coordinate or distance of the pen tip from the tablet surface. The TabletPropertyMetricUnit enumeration type determines the unit of measurement for this property.
+		PacketStatus, // Contains one or more of the following flag values: The cursor is touching the drawing surface (Value = 1). The cursor is inverted. For example, the eraser end of the pen is pointing toward the surface (Value = 2). Not used (Value = 4). The barrel button is pressed (Value = 8).
+		TimerTick, // The time the packet was generated.
+		SerialNumber, // The packet property for identifying the packet. This is the same value you use to retrieve the packet from the packet queue.
+		NormalPressure, // The pressure of the pen tip perpendicular to the tablet surface. The greater the pressure on the pen tip, the more ink that is drawn.
+		TangentPressure, // The pressure of the pen tip along the plane of the tablet surface.
+		ButtonPressure, // The pressure on a pressure sensitive button.
+		XTiltOrientation, // The angle between the y,z-plane and the pen and y-axis plane. Applies to a pen cursor. The value is 0 when the pen is perpendicular to the drawing surface and is positive when the pen is to the right of perpendicular.
+		YTiltOrientation, // The angle between the x,z-plane and the pen and x-axis plane. Applies to a pen cursor. The value is 0 when the pen is perpendicular to the drawing surface and is positive when the pen is upward or away from the user.
+		AzimuthOrientation, // The clockwise rotation of the cursor about the z-axis through a full circular range.
+		AltitudeOrientation, // The angle between the axis of the pen and the surface of the tablet. The value is 0 when the pen is parallel to the surface and 90 when the pen is perpendicular to the surface. The values are negative when the pen is inverted.
+		TwistOrientation, // The clockwise rotation of the cursor about its own axis.
+		PitchRotation, // The packet property that indicates whether the tip is above or below a horizontal line that is perpendicular to the writing surface. Note: This requires a 3D digitizer. The value is positive if the tip is above the line and negative if it is below the line. For example, if you hold the pen in front of you and write on an imaginary wall, the pitch is positive if the tip is above a line extending from you to the wall.
+		RollRotation, // The clockwise rotation of the pen around its own axis. Note: This requires a 3D digitizer.
+		YawRotation, // The angle of the pen to the left or right around the center of its horizontal axis when the pen is horizontal. Note: This requires a 3D digitizer. If you hold the pen in front of you and write on an imaginary wall, zero yaw indicates that the pen is perpendicular to the wall. The value is negative if the tip is to the left of perpendicular and positive if the tip is to the right of perpendicular.
+		Width, // The width of the contact area on a touch digitizer.
+		Height, // The height of the contact area on a touch digitizer.
+		FingerContactConfidence, // The level of confidence that there was finger contact on a touch digitizer.
+		DeviceContactID, // The device contact identifier for a packet.
+
+		Num_Enumerators, // THIS IS NOT A VALID ENUMERATOR, BUT IT CAN BE USED TO QUERY HOW MANY PACKET PROPERTIES THERE ARE.
+		Invalid_Enumerator = INDEX_NONE // THIS IS NOT A VALID ENUMERATOR, BUT IT CAN BE USED TO INDICATE THAT A PACKET PROPERTY ENTRY IS INVALID.
+	};
+
+	enum class ETabletPropertyMetricUnit : int8
+	{
+		Default, // Units are unknown.
+		Inches,
+		Centimeters,
+		Degrees,
+		Radians,
+		Seconds,
+		Pounds,
+		Grams,
+
+		Num_Enumerators, // THIS IS NOT A VALID ENUMERATOR, BUT IT CAN BE USED TO QUERY HOW MANY PACKET PROPERTIES THERE ARE.
+		Invalid_Enumerator = INDEX_NONE // THIS IS NOT A VALID ENUMERATOR, BUT IT CAN BE USED TO INDICATE THAT A PACKET PROPERTY ENTRY IS INVALID.
+	};
+
+	using FuncSetProperty = void(FStylusInputPacket&, int32, const int8*);
+
+	struct FPacketProperty
+	{
+		EPacketPropertyType Type = EPacketPropertyType::Invalid_Enumerator;
+		ETabletPropertyMetricUnit MetricUnit = ETabletPropertyMetricUnit::Default;
+		int32 Minimum = 0;
+		int32 Maximum = 0;
+		float Resolution = 0.0f;
+		FuncSetProperty* SetProperty = nullptr;
+
+		static constexpr int32 SetPropertyDataBufferLength = 24;
+		int8 SetPropertyData[SetPropertyDataBufferLength];
+	};
+
+	struct FPacketPropertyConstant
+	{
+		EPacketPropertyType PacketPropertyType;
+		//GUID Guid;
+		const wchar_t* StrGuid;
+	};
+
+	extern FPacketPropertyConstant PacketPropertyConstants[static_cast<uint32>(EPacketPropertyType::Num_Enumerators)];
+
+	struct FWindowContext
+	{
+		FVector2d XYScale;
+		FVector2d XYMaximum;
+		FVector2d WindowSize;
+		//FCocoaWindow* WindowHandle;
+	};
+
+	class FTabletContext final : public IStylusInputTabletContext
+	{
+	public:
+		explicit FTabletContext(const uint32 ID)
+			: ID(ID)
+		{}
+
+		virtual uint32 GetID() const override { return ID; }
+		virtual FString GetName() const override { return ProductName; }
+
+		virtual uint32 GetProductID() const { return ProductID; }
+		virtual FString GetVendorName() const { return VendorName; }
+		virtual uint32 GetVendorID() const { return VendorID; }
+		virtual FString GetSerialNumber() const { return SerialNumber; }
+
+		virtual FIntRect GetInputRectangle() const override { return InputRectangle; }
+		virtual ETabletHardwareCapabilities GetHardwareCapabilities() const override { return HardwareCapabilities; }
+		virtual ETabletSupportedProperties GetSupportedProperties() const override { return SupportedProperties; }
+
+		uint32 ID;
+		FString ProductName;
+
+		uint32 ProductID;
+		FString VendorName;
+		uint32 VendorID;
+		FString SerialNumber;
+		FIntRect InputRectangle;
+		ETabletSupportedProperties SupportedProperties = ETabletSupportedProperties::None;
+		ETabletHardwareCapabilities HardwareCapabilities = ETabletHardwareCapabilities::None;
+		FPacketProperty PacketDescriptions[static_cast<uint32>(EPacketPropertyType::Num_Enumerators)];
+	};
+
+	using FTabletContextContainer = Private::TSharedRefDataContainer<FTabletContext, false>;
+
+	class FStylusButton final : IStylusInputStylusButton
+	{
+	public:
+		virtual FString GetID() const override { return ID; }
+		virtual FString GetName() const override { return Name; }
+
+		FString ID;
+		FString Name;
+	};
+
+	class FStylusInfo final : public IStylusInputStylusInfo
+	{
+	public:
+		explicit FStylusInfo(const uint32 ID)
+			: ID(ID)
+		{}
+
+		virtual uint32 GetID() const override { return ID; }
+		virtual FString GetName() const override { return Name; }
+		virtual uint32 GetNumButtons() const override { return Buttons.Num(); }
+
+		virtual const IStylusInputStylusButton* GetButton(int32 Index) const override
+		{
+			return 0 <= Index && Index < Buttons.Num() ? reinterpret_cast<const IStylusInputStylusButton*>(&Buttons[Index]) : nullptr;
+		}
+
+		uint32 ID = 0;
+		FString Name;
+		TArray<FStylusButton> Buttons;
+	};
+
+	using FStylusInfoContainer = Private::TSharedRefDataContainer<FStylusInfo, false>;
+	bool GetMacTabletDevices(FTabletContextContainer& OutDevices);
+}
diff --git a/Engine/Plugins/Editor/StylusInput/Source/StylusInputMac/Private/NSEventHandler.cpp b/Engine/Plugins/Editor/StylusInput/Source/StylusInputMac/Private/NSEventHandler.cpp
new file mode 100644
--- /dev/null
+++ b/Engine/Plugins/Editor/StylusInput/Source/StylusInputMac/Private/NSEventHandler.cpp
@@ -0,0 +1,331 @@
+// Copyright Epic Games, Inc. All Rights Reserved.
+
+#include "NSEventHandler.h"
+
+#include <StylusInput.h>
+#include <StylusInputUtils.h>
+#include <Mac/CocoaWindow.h>
+
+#define LOG_PREAMBLE "NSEventHandler"
+
+using namespace UE::StylusInput::Private;
+
+namespace UE::StylusInput::Mac
+{
+	//Info shared by IOHID and NSEventLoop
+	static uint32 CurrentTabletContextID = 0;
+	static EPenStatus CurrentPenStatus;
+	static EPacketType CurrentPacketType;
+	
+	FNSEventHandler::FNSEventHandler(IStylusInputInstance* Instance,
+						   const FCocoaWindow* CocoaWindow,
+						   IStylusInputEventHandler* EventHandler)
+		: Instance(Instance)
+		, CocoaWindow(CocoaWindow)
+	{
+		if (EventHandler)
+		{
+			// Immediately install an event handler during construction to capture events coming through during initialization.
+			AddEventHandler(EventHandler);
+		}
+	}
+
+	FNSEventHandler::~FNSEventHandler()
+	{
+	}
+
+	void FNSEventHandler::StartListen()
+	{
+		// --- IOHID Setup ---
+		// Handling IOHID events to get the device ID of the event sent
+		if( IsIOHIDAvailable() )
+		{
+			LogVerbose(LOG_PREAMBLE, "Start IOHID");
+			IOHIDManagerRef hidManager = IOHIDManagerCreate(kCFAllocatorDefault, kIOHIDOptionsTypeNone);
+
+			IOHIDManagerSetDeviceMatching(hidManager, NULL);
+
+			// Schedule with runloop
+			IOHIDManagerScheduleWithRunLoop(hidManager, CFRunLoopGetCurrent(), kCFRunLoopDefaultMode);
+
+			// Open manager
+			IOReturn openResult = IOHIDManagerOpen(hidManager, kIOHIDOptionsTypeNone);
+			if (openResult != kIOReturnSuccess) {
+				NSLog(@"Failed to open IOHIDManager: 0x%08x", openResult);
+			}
+
+			// Register callback
+			IOHIDManagerRegisterInputValueCallback(hidManager, HandleHIDEvent, this);
+
+			// Save reference if you want to stop later
+			this->HIDManagerRef = hidManager;
+		}
+
+		// --- NSEvent Setup ---
+		LogVerbose(LOG_PREAMBLE, "Start NSEvent");
+		ensure(TabletContextListener == nil);
+
+		NSEventMask mask = NSEventMaskTabletPoint | NSEventMaskTabletProximity | NSEventMaskLeftMouseDragged | NSEventMaskRightMouseDragged | NSEventMaskOtherMouseDragged | NSEventMaskMouseMoved;
+		TabletContextListener = ( [NSEvent addLocalMonitorForEventsMatchingMask:mask handler:^(NSEvent* Event) {
+			//if( Event.window == Instance->GetWindowContext()->Window )
+			return HandleNSEvent(Event); }] );
+	}
+
+	void FNSEventHandler::StopListen()
+	{
+		if( IsIOHIDAvailable() && HIDManagerRef )
+		{
+			IOHIDManagerUnscheduleFromRunLoop(HIDManagerRef, CFRunLoopGetCurrent(), kCFRunLoopDefaultMode);
+			IOHIDManagerClose(HIDManagerRef, kIOHIDOptionsTypeNone);
+			CFRelease(HIDManagerRef);
+			HIDManagerRef = nullptr;
+		}
+
+		if( TabletContextListener == nil )
+			return;
+
+		[NSEvent removeMonitor:TabletContextListener];
+		TabletContextListener = nil;
+	}
+
+	bool FNSEventHandler::AddEventHandler(IStylusInputEventHandler* EventHandler)
+	{
+		check(EventHandler);
+
+		if (EventHandlers.Contains(EventHandler))
+		{
+			LogWarning(LOG_PREAMBLE, FString::Format(TEXT("Event handler '{0}' already exists."), {EventHandler->GetName()}));
+			return false;
+		}
+
+		EventHandlers.Add(EventHandler);
+
+		LogVerbose(LOG_PREAMBLE, FString::Format(TEXT("Event handler '{0}' was added."), {EventHandler->GetName()}));
+
+		return true;
+	}
+
+
+	bool FNSEventHandler::RemoveEventHandler(IStylusInputEventHandler* EventHandler)
+	{
+		check(EventHandler);
+
+		const bool bWasRemoved = EventHandlers.Remove(EventHandler) > 0;
+
+		if (bWasRemoved)
+		{
+			LogVerbose(LOG_PREAMBLE, FString::Format(TEXT("Event handler '{0}' was removed."), {EventHandler->GetName()}));
+		}
+
+		return bWasRemoved;
+	}
+
+	FTabletContextContainer& FNSEventHandler::GetTabletContexts()
+	{
+		return TabletContexts;
+	}
+
+	void FNSEventHandler::DebugEvent(const FString& Message) const
+	{
+		for (IStylusInputEventHandler* EventHandler : EventHandlers)
+		{
+			EventHandler->OnDebugEvent(Message, Instance);
+		}
+	}
+	
+	bool FNSEventHandler::IsIOHIDAvailable()
+	{
+		IOHIDManagerRef manager = IOHIDManagerCreate(kCFAllocatorDefault, kIOHIDOptionsTypeNone);
+		if (!manager) 
+		{
+			NSLog(@"IOHIDManager could not be created.");
+			return false;
+		}
+
+		// Try opening the manager
+		IOReturn result = IOHIDManagerOpen(manager, kIOHIDOptionsTypeNone);
+		if (result != kIOReturnSuccess) 
+		{
+			NSLog(@"IOHIDManagerOpen failed: 0x%08x", result);
+			CFRelease(manager);
+			return false;
+		}
+
+		IOHIDManagerSetDeviceMatching(manager, NULL); // NULL matches all devices
+
+		CFSetRef deviceSet = IOHIDManagerCopyDevices(manager);
+		bool hasDevices = false;
+		if (deviceSet && CFGetTypeID(deviceSet) == CFSetGetTypeID()) 
+		{
+			CFIndex count = CFSetGetCount(deviceSet);
+			hasDevices = (count > 0);
+		}
+
+		if (!hasDevices)
+		{
+			NSLog(@"No HID devices found or permission denied.");
+		}
+
+		IOHIDManagerClose(manager, kIOHIDOptionsTypeNone);
+		CFRelease(manager);
+
+		return hasDevices;
+	}
+
+	void FNSEventHandler::HandleHIDEvent(void* context, IOReturn result, void* sender, IOHIDValueRef value)
+	{
+		// Retrieve the device
+		IOHIDDeviceRef device = IOHIDElementGetDevice(IOHIDValueGetElement(value));
+
+		CFNumberRef productID = (CFNumberRef)IOHIDDeviceGetProperty(device, CFSTR(kIOHIDProductIDKey));
+
+		uint32 ID = 0;
+		if (productID)
+			CFNumberGetValue(productID, kCFNumberSInt32Type, &ID);
+
+		CurrentTabletContextID = ID;
+
+		IOHIDElementRef elem = IOHIDValueGetElement(value);
+		uint32_t elementUsagePage = IOHIDElementGetUsagePage(elem);
+		uint32_t elementUsage = IOHIDElementGetUsage(elem);
+		CFIndex state = IOHIDValueGetIntegerValue(value);
+
+		if (elementUsagePage == kHIDPage_Digitizer || elementUsagePage >= 0xFF00)
+		{
+			switch (elementUsage)
+			{
+			case kHIDUsage_Dig_TipSwitch:
+				if (state) CurrentPenStatus = CurrentPenStatus | EPenStatus::CursorIsTouching;
+				else CurrentPenStatus = CurrentPenStatus & ~EPenStatus::CursorIsTouching;
+				break;
+			case kHIDUsage_Dig_BarrelSwitch:
+				if (state) CurrentPenStatus = CurrentPenStatus | EPenStatus::BarrelButtonPressed;
+				else CurrentPenStatus = CurrentPenStatus & ~EPenStatus::BarrelButtonPressed;
+				break;
+			case kHIDUsage_Dig_Eraser:
+				if (state) CurrentPenStatus = CurrentPenStatus | EPenStatus::CursorIsInverted;
+				else CurrentPenStatus = CurrentPenStatus & ~EPenStatus::CursorIsInverted;
+				break;
+			}
+		}
+
+		// Caveat: Internal trackpad doesn't send events via IOHID, meaning we can't grab its ID from this function
+	}
+
+	//Set the Altitude and Azimuth fields based on the Tilt field
+	void TiltToOrientation(NSPoint Tilt, float& outAzimuth, float& outAltitude)
+	{
+		outAzimuth = 0;
+		if (Tilt.x != 0)
+		{
+			outAzimuth = PI / 2 - FMath::Atan2(FMath::Cos(Tilt.x) * FMath::Sin(Tilt.y), FMath::Cos(Tilt.y) * FMath::Sin(Tilt.x));
+			if (outAzimuth < 0)
+				outAzimuth += 2 * PI;
+		}
+
+		outAltitude = PI / 2 - FMath::Acos(FMath::Cos(Tilt.x) * FMath::Cos(Tilt.y));
+
+		outAltitude = FMath::RadiansToDegrees(outAltitude);
+		outAzimuth = FMath::RadiansToDegrees(outAzimuth);
+	}
+
+	NSEvent* FNSEventHandler::HandleNSEvent(NSEvent* Event)
+	{
+		ensure(Event);
+
+		if (!Event.window)
+			return Event;
+
+		PacketStats.NewPacket();
+		FStylusInputPacket Packet;
+
+		NSPoint LocalPosition;
+		ConvertCocoaEventPositionToSlate(Event, &LocalPosition);
+
+		Packet.X = LocalPosition.x;
+		Packet.Y = LocalPosition.y;
+		Packet.Z = Event.absoluteZ;
+
+		Packet.TimerTick = (int32)Event.timestamp;
+		Packet.SerialNumber = (int32)Event.pointingDeviceSerialNumber;
+
+		Packet.NormalPressure = Event.pressure;
+		if (Event.pressure == 0)
+		{
+			CurrentPenStatus = CurrentPenStatus & ~EPenStatus::CursorIsTouching;
+			if (CurrentPacketType == EPacketType::OnDigitizer)
+				CurrentPacketType = EPacketType::StylusUp;
+		}
+
+		Packet.TangentPressure = Event.tangentialPressure;
+
+		float Azimuth = 0.f;
+		float Altitude = 0.f;
+		TiltToOrientation(Event.tilt, Azimuth, Altitude);
+
+		Packet.AzimuthOrientation = Azimuth;
+		Packet.AltitudeOrientation = Altitude;
+		Packet.TwistOrientation = Event.rotation;
+
+		Packet.XTiltOrientation = Event.tilt.x;
+		Packet.YTiltOrientation = Event.tilt.y;
+
+		Packet.TabletContextID = CurrentTabletContextID;
+
+		if (Event.buttonMask & NSEventButtonMaskPenTip) // Inconvenient warning
+		{
+			CurrentPenStatus = CurrentPenStatus | EPenStatus::CursorIsTouching;
+			if (CurrentPacketType != EPacketType::OnDigitizer)
+				CurrentPacketType = EPacketType::StylusDown;
+			else
+				CurrentPacketType = EPacketType::OnDigitizer;
+		}
+		else
+		{
+			CurrentPenStatus = CurrentPenStatus & ~EPenStatus::CursorIsTouching;
+		}
+
+		if (Event.buttonMask & NSEventButtonMaskPenLowerSide) // Inconvenient warning
+		{
+			CurrentPenStatus = CurrentPenStatus | EPenStatus::BarrelButtonPressed;
+		}
+		else
+		{
+			CurrentPenStatus = CurrentPenStatus & ~EPenStatus::BarrelButtonPressed;
+		}
+
+		Packet.PenStatus = CurrentPenStatus;
+		Packet.Type = CurrentPacketType;
+
+		for (IStylusInputEventHandler* EventHandler : EventHandlers)
+		{
+			EventHandler->OnPacket(Packet, Instance);
+		}
+		return Event;
+	}
+
+	bool FNSEventHandler::ConvertCocoaEventPositionToSlate(NSEvent* Event, NSPoint* oPoint)
+	{
+		if (!Event || !oPoint)
+			return false;
+
+		if (!CocoaWindow)
+			return false;
+
+		*oPoint = Event.locationInWindow;
+		// oPoint is relative to the window it is hovering above
+		// convert to the window this plugin is used on
+		*oPoint = [Event.window convertPointToScreen:*oPoint];
+		*oPoint = [CocoaWindow convertPointFromScreen:*oPoint];
+		NSSize windowSize = CocoaWindow.contentView.frame.size;
+
+		CGFloat scale = CocoaWindow.backingScaleFactor;
+
+		oPoint->x *= scale;
+		oPoint->y = (windowSize.height - oPoint->y) * scale;
+
+		return true;
+	}
+}
+
+#undef LOG_PREAMBLE
diff --git a/Engine/Plugins/Editor/StylusInput/Source/StylusInputMac/Private/NSEventHandler.h b/Engine/Plugins/Editor/StylusInput/Source/StylusInputMac/Private/NSEventHandler.h
new file mode 100644
--- /dev/null
+++ b/Engine/Plugins/Editor/StylusInput/Source/StylusInputMac/Private/NSEventHandler.h
@@ -0,0 +1,68 @@
+// Copyright Epic Games, Inc. All Rights Reserved.
+
+#pragma once
+
+#include <StylusInputPacket.h>
+#include <Delegates/Delegate.h>
+
+#include "MacStats.h"
+#include "MacTabletContext.h"
+
+#import <IOKit/hid/IOHIDManager.h>
+#import <Foundation/Foundation.h>
+
+@class FCocoaWindow;
+
+namespace UE::StylusInput
+{
+	class IStylusInputEventHandler;
+	class IStylusInputInstance;
+}
+
+namespace UE::StylusInput::Mac
+{
+	class FNSEventHandler
+	{
+	public:
+		FNSEventHandler(IStylusInputInstance* Instance, const FCocoaWindow* CocoaWindow, IStylusInputEventHandler* EventHandler);
+		virtual ~FNSEventHandler();
+		
+		bool AddEventHandler(IStylusInputEventHandler* EventHandler);
+		bool RemoveEventHandler(IStylusInputEventHandler* EventHandler);
+		int32 NumEventHandlers() const { return EventHandlers.Num(); }
+
+		float GetPacketsPerSecond() const { return PacketStats.GetPacketsPerSecond(); }
+
+		// Starts listening to events on all tablet contexts we have
+		virtual void StartListen();
+
+		// Stops listening to events on all tablet contexts we have
+		virtual void StopListen();
+
+		FTabletContextContainer& GetTabletContexts();
+
+		void SetTabletContextID(uint32 ID);
+		void SetCurrentPenStatus(EPenStatus);
+		void SetCurrentPacketType(EPacketType);
+
+	protected:
+		void DebugEvent(const FString& Message) const;
+
+		// Used to get the device ID of the device that sent the last event. We treat the actual event with NSEvent
+		static bool IsIOHIDAvailable();
+		static void HandleHIDEvent(void* context, IOReturn result, void* sender, IOHIDValueRef value);
+		virtual NSEvent* HandleNSEvent(NSEvent* Event);
+
+		bool ConvertCocoaEventPositionToSlate(NSEvent* Event, NSPoint* oPoint);
+
+	protected:
+		IStylusInputInstance* const Instance;
+		FPacketStats PacketStats;
+		const FCocoaWindow* CocoaWindow;
+		FTabletContextContainer TabletContexts;
+		TArray<IStylusInputEventHandler*> EventHandlers;
+
+		id TabletContextListener = nil;
+		IOHIDManagerRef HIDManagerRef;
+	};
+}
diff --git a/Engine/Plugins/Editor/StylusInput/Source/StylusInputMac/Private/StylusInputMacModule.cpp b/Engine/Plugins/Editor/StylusInput/Source/StylusInputMac/Private/StylusInputMacModule.cpp
new file mode 100644
--- /dev/null
+++ b/Engine/Plugins/Editor/StylusInput/Source/StylusInputMac/Private/StylusInputMacModule.cpp
@@ -0,0 +1,39 @@
+﻿// Copyright Epic Games, Inc. All Rights Reserved.
+
+#include <Modules/ModuleManager.h>
+
+#include "MacInterface.h"
+
+namespace UE::StylusInput::Mac
+{
+	class FStylusInputMacModule : public IModuleInterface
+	{
+	public:
+		virtual void StartupModule() override;
+		virtual void ShutdownModule() override;
+
+	private:
+		TUniquePtr<IStylusInputInterface> RegisteredInterface;
+	};
+
+	void FStylusInputMacModule::StartupModule()
+	{
+		RegisteredInterface = FMacInterface::Create();
+		if (RegisteredInterface)
+		{
+			RegisterInterface(RegisteredInterface.Get());
+		}
+	}
+
+	void FStylusInputMacModule::ShutdownModule()
+	{
+		if (RegisteredInterface)
+		{
+			UnregisterInterface(RegisteredInterface.Get());
+
+			RegisteredInterface.Reset();
+		}
+	}
+}
+
+IMPLEMENT_MODULE(UE::StylusInput::Mac::FStylusInputMacModule, StylusInputMac);
diff --git a/Engine/Plugins/Editor/StylusInput/Source/StylusInputMac/StylusInputMac.Build.cs b/Engine/Plugins/Editor/StylusInput/Source/StylusInputMac/StylusInputMac.Build.cs
new file mode 100644
--- /dev/null
+++ b/Engine/Plugins/Editor/StylusInput/Source/StylusInputMac/StylusInputMac.Build.cs
@@ -0,0 +1,27 @@
+﻿// Copyright Epic Games, Inc. All Rights Reserved.
+
+namespace UnrealBuildTool.Rules
+{
+	public class StylusInputMac : ModuleRules
+	{
+		public StylusInputMac(ReadOnlyTargetRules Target) : base(Target)
+		{
+			PrivateDependencyModuleNames.AddRange(
+				[
+					"Core",
+					"Slate",
+					"SlateCore",
+					"StylusInput"
+				]
+			);
+
+			PublicFrameworks.AddRange(
+				[
+					"IOKit",
+					"CoreFoundation",
+					"Foundation"
+				]
+			);
+		}
+	}
+}
\ No newline at end of file
diff --git a/Engine/Plugins/Editor/StylusInput/StylusInput.uplugin b/Engine/Plugins/Editor/StylusInput/StylusInput.uplugin
--- a/Engine/Plugins/Editor/StylusInput/StylusInput.uplugin
+++ b/Engine/Plugins/Editor/StylusInput/StylusInput.uplugin
@@ -22,21 +22,27 @@
 			"LoadingPhase" : "Default"
 		},
 		{
-			"Name" : "StylusInputRealTimeStylus",
+			"Name" : "StylusInputDebugWidget",
 			"Type" : "EditorNoCommandlet",
-			"PlatformAllowList": ["Win64"],
+			"LoadingPhase" : "PostEngineInit"
+		},
+		{
+			"Name" : "StylusInputMac",
+			"Type" : "EditorNoCommandlet",
+			"PlatformAllowList": ["Mac"],
 			"LoadingPhase" : "PostDefault"
 		},
 		{
-			"Name" : "StylusInputWintab",
+			"Name" : "StylusInputRealTimeStylus",
 			"Type" : "EditorNoCommandlet",
 			"PlatformAllowList": ["Win64"],
 			"LoadingPhase" : "PostDefault"
 		},
 		{
-			"Name" : "StylusInputDebugWidget",
+			"Name" : "StylusInputWintab",
 			"Type" : "EditorNoCommandlet",
-			"LoadingPhase" : "PostEngineInit"
+			"PlatformAllowList": ["Win64"],
+			"LoadingPhase" : "PostDefault"
 		}
 	]
 }
