Minor fixes for Mac Stylus

#rnx

diff --git a/Engine/Plugins/Editor/StylusInput/Source/StylusInputMac/Private/MacTabletContext.cpp b/Engine/Plugins/Editor/StylusInput/Source/StylusInputMac/Private/MacTabletContext.cpp
--- a/Engine/Plugins/Editor/StylusInput/Source/StylusInputMac/Private/MacTabletContext.cpp
+++ b/Engine/Plugins/Editor/StylusInput/Source/StylusInputMac/Private/MacTabletContext.cpp
@@ -116,13 +116,13 @@ namespace UE::StylusInput::Mac
 							case kHIDUsage_GD_X:
 								context->SupportedProperties = context->SupportedProperties | ETabletSupportedProperties::X;
 								context->InputRectangle.Min[0] = IOHIDElementGetLogicalMin(element);
-								context->InputRectangle.Max[0] = IOHIDElementGetLogicalMax(element) - context->InputRectangle.Min[0];
+								context->InputRectangle.Max[0] = IOHIDElementGetLogicalMax(element);
 								NSLog(@"Supported X");
 								break;
 							case kHIDUsage_GD_Y:
 								context->SupportedProperties = context->SupportedProperties | ETabletSupportedProperties::Y;
 								context->InputRectangle.Min[1] = IOHIDElementGetLogicalMin(element);
-								context->InputRectangle.Max[1] = IOHIDElementGetLogicalMax(element) - context->InputRectangle.Min[1];
+								context->InputRectangle.Max[1] = IOHIDElementGetLogicalMax(element);
 								NSLog(@"Supported Y");
 								break;
 							case kHIDUsage_GD_Z:
@@ -159,6 +159,7 @@ namespace UE::StylusInput::Mac
 								context->SupportedProperties = context->SupportedProperties | ETabletSupportedProperties::PacketStatus;
 								context->HardwareCapabilities = context->HardwareCapabilities | ETabletHardwareCapabilities::Integrated;
 								NSLog(@"Touch Supported");
+								break;
 							case kHIDUsage_Dig_TipPressure:
 								context->SupportedProperties = context->SupportedProperties | ETabletSupportedProperties::NormalPressure;
 								NSLog(@"Pressure Supported");
@@ -208,6 +209,7 @@ namespace UE::StylusInput::Mac
 						}
 					}
 				}
+				CFRelease(elements);
 			}
 		}
 		CFRelease(DeviceSet);
diff --git a/Engine/Plugins/Editor/StylusInput/Source/StylusInputMac/Private/NSEventHandler.cpp b/Engine/Plugins/Editor/StylusInput/Source/StylusInputMac/Private/NSEventHandler.cpp
--- a/Engine/Plugins/Editor/StylusInput/Source/StylusInputMac/Private/NSEventHandler.cpp
+++ b/Engine/Plugins/Editor/StylusInput/Source/StylusInputMac/Private/NSEventHandler.cpp
@@ -22,6 +22,7 @@ namespace UE::StylusInput::Mac
 						   IStylusInputEventHandler* EventHandler)
 		: Instance(Instance)
 		, CocoaWindow(CocoaWindow)
+		, HIDManagerRef(nullptr)
 	{
 		if (EventHandler)
 		{
@@ -32,6 +33,11 @@ namespace UE::StylusInput::Mac
 
 	FNSEventHandler::~FNSEventHandler()
 	{
+		if (HIDManagerRef)
+		{
+			CFRelease(HIDManagerRef);
+			HIDManagerRef = nullptr;
+		}
 	}
 
 	void FNSEventHandler::StartListen()
@@ -42,6 +48,7 @@ namespace UE::StylusInput::Mac
 		{
 			LogVerbose(LOG_PREAMBLE, "Start IOHID");
 			IOHIDManagerRef hidManager = IOHIDManagerCreate(kCFAllocatorDefault, kIOHIDOptionsTypeNone);
+			ensure(hidManager);
 
 			IOHIDManagerSetDeviceMatching(hidManager, NULL);
 
@@ -165,6 +172,12 @@ namespace UE::StylusInput::Mac
 		{
 			NSLog(@"No HID devices found or permission denied.");
 		}
+		
+		if (deviceSet)
+		{
+			CFRelease(deviceSet);
+			deviceSet = nullptr;
+		}
 
 		IOHIDManagerClose(manager, kIOHIDOptionsTypeNone);
 		CFRelease(manager);
