StylusInput: Fix debug widget not receiving packets after widget gets attached to a different window

#rnx
#FailedValidations Preflight

diff --git a/Engine/Plugins/Editor/StylusInput/Source/StylusInputDebugWidget/Private/StylusInputDebugWidget.cpp b/Engine/Plugins/Editor/StylusInput/Source/StylusInputDebugWidget/Private/StylusInputDebugWidget.cpp
--- a/Engine/Plugins/Editor/StylusInput/Source/StylusInputDebugWidget/Private/StylusInputDebugWidget.cpp
+++ b/Engine/Plugins/Editor/StylusInput/Source/StylusInputDebugWidget/Private/StylusInputDebugWidget.cpp
@@ -497,6 +497,27 @@ namespace UE::StylusInput::DebugWidget
 		                           });
 	}
 
+	void SStylusInputDebugWidget::NotifyWidgetRelocated()
+	{
+		if (const TSharedPtr<SWindow> PrevWindow = StylusInputWindow.Pin())
+		{
+			if (const TSharedPtr<SWindow> CurrentWindow = FSlateApplication::Get().FindWidgetWindow(AsShared()))
+			{
+				if (PrevWindow.Get() == CurrentWindow.Get())
+				{
+					// Window did not change; no need to re-initialize stylus input.
+					return;
+				}
+			}
+		}
+
+		UnregisterEventHandler();
+		ReleaseStylusInput();
+
+		AcquireStylusInput();
+		RegisterEventHandler();
+	}
+
 	void SStylusInputDebugWidget::AcquireStylusInput()
 	{
 		if (StylusInput)
@@ -518,7 +539,10 @@ namespace UE::StylusInput::DebugWidget
 		if (!StylusInput)
 		{
 			LogError("Could not acquire stylus input instance.");
+			return;
 		}
+
+		StylusInputWindow = Window.ToWeakPtr();
 	}
 
 	void SStylusInputDebugWidget::ReleaseStylusInput()
@@ -537,6 +561,8 @@ namespace UE::StylusInput::DebugWidget
 
 			StylusInput = nullptr;
 		}
+
+		StylusInputWindow = nullptr;
 	}
 
 	void SStylusInputDebugWidget::RegisterEventHandler()
diff --git a/Engine/Plugins/Editor/StylusInput/Source/StylusInputDebugWidget/Private/StylusInputDebugWidget.h b/Engine/Plugins/Editor/StylusInput/Source/StylusInputDebugWidget/Private/StylusInputDebugWidget.h
--- a/Engine/Plugins/Editor/StylusInput/Source/StylusInputDebugWidget/Private/StylusInputDebugWidget.h
+++ b/Engine/Plugins/Editor/StylusInput/Source/StylusInputDebugWidget/Private/StylusInputDebugWidget.h
@@ -62,6 +62,8 @@ namespace UE::StylusInput::DebugWidget
 
 		void Construct(const FArguments& Args);
 
+		void NotifyWidgetRelocated();
+
 	private:
 		void AcquireStylusInput();
 		void ReleaseStylusInput();
@@ -84,6 +86,7 @@ namespace UE::StylusInput::DebugWidget
 
 		IStylusInputInstance* StylusInput = nullptr;
 		TUniquePtr<IStylusInputEventHandler> EventHandler;
+		TWeakPtr<SWindow> StylusInputWindow;
 
 		TSharedPtr<SStylusInputDebugPaintWidget> PaintWidget;
 
diff --git a/Engine/Plugins/Editor/StylusInput/Source/StylusInputDebugWidget/Private/StylusInputDebugWidgetModule.cpp b/Engine/Plugins/Editor/StylusInput/Source/StylusInputDebugWidget/Private/StylusInputDebugWidgetModule.cpp
--- a/Engine/Plugins/Editor/StylusInput/Source/StylusInputDebugWidget/Private/StylusInputDebugWidgetModule.cpp
+++ b/Engine/Plugins/Editor/StylusInput/Source/StylusInputDebugWidget/Private/StylusInputDebugWidgetModule.cpp
@@ -78,6 +78,16 @@ namespace UE::StylusInput::DebugWidget
 	{
 		TSharedRef<SDockTab> Tab = SNew(SDockTab).TabRole(NomadTab);
 		Tab->SetContent(GetDebugWidget());
+
+		// Relocating the tab might mean that it is parented under a new window, which would require re-initializing the stylus input instance.
+		Tab->SetOnTabRelocated(FSimpleDelegate::CreateLambda([DebugWidgetPtr = DebugWidgetPtr]
+		{
+			if (const TSharedPtr<SStylusInputDebugWidget> DebugWidget = DebugWidgetPtr.Pin())
+			{
+				DebugWidget->NotifyWidgetRelocated();
+			}
+		}));
+
 		return Tab;
 	}
 
